Fix crash in PhAT when deleting a body that has more than one primitive of the same type

diff --git a/Engine/Source/Editor/PhAT/Private/PhATSharedData.cpp b/Engine/Source/Editor/PhAT/Private/PhATSharedData.cpp
--- a/Engine/Source/Editor/PhAT/Private/PhATSharedData.cpp
+++ b/Engine/Source/Editor/PhAT/Private/PhATSharedData.cpp
@@ -1082,56 +1082,65 @@ void FPhATSharedData::DeleteCurrentPrim()
 
 	//We will first get all the bodysetups we're interested in. The number of duplicates each bodysetup has tells us how many geoms are being deleted
 	//We need to do this first because deleting will modify our selection
-	TArray<FSelection> BodySelections;
+	TMap<UBodySetup *, TArray<FSelection>> BodySelectionMap;
 	TArray<UBodySetup*> BodySetups;
 	for(int32 i=0; i<SelectedBodies.Num(); ++i)
 	{
 		UBodySetup* BodySetup = PhysicsAsset->BodySetup[SelectedBodies[i].Index];
-		BodySetups.Add(BodySetup);
-		BodySelections.Add(SelectedBodies[i]);
+		BodySelectionMap.FindOrAdd(BodySetup).Add(SelectedBodies[i]);
 	}
 
 	const FScopedTransaction Transaction( NSLOCTEXT("UnrealEd", "DeletePrimitive", "Delete Primitive") );
 
-	TArray<UBodySetup*> ToDelete;
-	for(int32 i=0; i<BodySetups.Num(); ++i)
+	for (TMap<UBodySetup*, TArray<FSelection> >::TConstIterator It(BodySelectionMap); It; ++It)
 	{
-		const FSelection & SelectedBody = BodySelections[i];
-		UBodySetup * BodySetup = BodySetups[i];
-		int32 BodyIndex = PhysicsAsset->FindBodyIndex(BodySetup->BoneName);
+		UBodySetup * BodySetup = It.Key();
+		const TArray<FSelection> & SelectedPrimitives = It.Value();
 
-		BodySetup->Modify();
+		int32 SphereDeletedCount = 0;
+		int32 BoxDeletedCount = 0;
+		int32 SphylDeletedCount = 0;
+		int32 ConvexDeletedCount = 0;
 
-		if (SelectedBody.PrimitiveType == KPT_Sphere)
+		for (int32 i = 0; i < SelectedPrimitives.Num(); ++i)
 		{
-			BodySetup->AggGeom.SphereElems.RemoveAt(SelectedBody.PrimitiveIndex);
-		}
-		else if (SelectedBody.PrimitiveType == KPT_Box)
-		{
-			BodySetup->AggGeom.BoxElems.RemoveAt(SelectedBody.PrimitiveIndex);
-		}
-		else if (SelectedBody.PrimitiveType == KPT_Sphyl)
-		{
-			BodySetup->AggGeom.SphylElems.RemoveAt(SelectedBody.PrimitiveIndex);
-		}
-		else if (SelectedBody.PrimitiveType == KPT_Convex)
-		{
-			BodySetup->AggGeom.ConvexElems.RemoveAt(SelectedBody.PrimitiveIndex);
-			// Need to invalidate GUID in this case as cooked data must be updated
-			BodySetup->InvalidatePhysicsData();
-		}
+			const FSelection & SelectedBody = SelectedPrimitives[i];
+			int32 BodyIndex = PhysicsAsset->FindBodyIndex(BodySetup->BoneName);
 
-		// If this bone has no more geometry - remove it totally.
-		if (BodySetup->AggGeom.GetElementCount() == 0)
-		{	
-			if(BodyIndex != INDEX_NONE)
+			BodySetup->Modify();
+
+			if (SelectedBody.PrimitiveType == KPT_Sphere)
+			{
+				BodySetup->AggGeom.SphereElems.RemoveAt(SelectedBody.PrimitiveIndex - (SphereDeletedCount++));
+			}
+			else if (SelectedBody.PrimitiveType == KPT_Box)
+			{
+				BodySetup->AggGeom.BoxElems.RemoveAt(SelectedBody.PrimitiveIndex - (BoxDeletedCount++));
+			}
+			else if (SelectedBody.PrimitiveType == KPT_Sphyl)
+			{
+				BodySetup->AggGeom.SphylElems.RemoveAt(SelectedBody.PrimitiveIndex - (SphylDeletedCount++));
+			}
+			else if (SelectedBody.PrimitiveType == KPT_Convex)
 			{
-				DeleteBody(BodyIndex, false);
+				BodySetup->AggGeom.ConvexElems.RemoveAt(SelectedBody.PrimitiveIndex - (ConvexDeletedCount++));
+				// Need to invalidate GUID in this case as cooked data must be updated
+				BodySetup->InvalidatePhysicsData();
 			}
 
-			if(CopiedBodySetup == BodySetup)
+			// If this bone has no more geometry - remove it totally.
+			if (BodySetup->AggGeom.GetElementCount() == 0)
 			{
-				CopiedBodySetup = NULL;
+				check(i == SelectedPrimitives.Num() - 1);	//we should really only delete on last prim - only reason this is even in for loop is because of API needing body index
+				if (BodyIndex != INDEX_NONE)
+				{
+					DeleteBody(BodyIndex, false);
+				}
+
+				if (CopiedBodySetup == BodySetup)
+				{
+					CopiedBodySetup = NULL;
+				}
 			}
 		}
 	}
