Fix crash when updating physics asset of vehicle component

diff --git a/Engine/Source/Editor/PhAT/Private/PhATSharedData.cpp b/Engine/Source/Editor/PhAT/Private/PhATSharedData.cpp
--- a/Engine/Source/Editor/PhAT/Private/PhATSharedData.cpp
+++ b/Engine/Source/Editor/PhAT/Private/PhATSharedData.cpp
@@ -284,6 +284,20 @@ void FPhATSharedData::RefreshPhysicsAssetChange(const UPhysicsAsset * InPhysAsse
 			}
 		}
 
+		for (FObjectIterator Iter(UWheeledVehicleMovementComponent::StaticClass()); Iter; ++Iter)
+		{
+			UWheeledVehicleMovementComponent * WheeledVehicleMovementComponent = Cast<UWheeledVehicleMovementComponent>(*Iter);
+			if (USkeletalMeshComponent * SkeltalMeshComponent = Cast<USkeletalMeshComponent>(WheeledVehicleMovementComponent->UpdatedComponent))
+			{
+				if (SkeltalMeshComponent->GetPhysicsAsset() == InPhysAsset)
+				{
+					//Need to recreate car data
+					WheeledVehicleMovementComponent->RecreatePhysicsState();
+				}
+
+			}
+		}
+
 		FEditorSupportDelegates::RedrawAllViewports.Broadcast();
 		// since we recreate physicsstate, a lot of transient state data will be gone
 		// so have to turn simulation off again. 
