PhAT action typo

diff --git a/Engine/Source/Editor/PhAT/Private/PhATActions.cpp b/Engine/Source/Editor/PhAT/Private/PhATActions.cpp
--- a/Engine/Source/Editor/PhAT/Private/PhATActions.cpp
+++ b/Engine/Source/Editor/PhAT/Private/PhATActions.cpp
@@ -9,7 +9,7 @@ void FPhATCommands::RegisterCommands()
 {
 	UI_COMMAND(ChangeDefaultMesh, "Change Mesh", "Change Default SkeletalMesh", EUserInterfaceActionType::Button, FInputGesture());
 	UI_COMMAND(ResetEntireAsset, "Reset Asset", "Opens Body Creation Settings And Replaces Asset Using The New Settings", EUserInterfaceActionType::Button, FInputGesture());
-	UI_COMMAND(RestetBoneCollision, "Reset", "Resest Bone Collision", EUserInterfaceActionType::Button, FInputGesture());
+	UI_COMMAND(RestetBoneCollision, "Reset", "Reset Bone Collision", EUserInterfaceActionType::Button, FInputGesture());
 	UI_COMMAND(ApplyPhysicalMaterial, "Apply PhysMat", "Applies Currently Selected Physical Material In Content Browser To All Bodies", EUserInterfaceActionType::Button, FInputGesture());
 	UI_COMMAND(EditingMode_Body, "Body Mode", "Body Editing Mode", EUserInterfaceActionType::RadioButton, FInputGesture());
 	UI_COMMAND(EditingMode_Constraint, "Constraint Mode", "Constraint Editing Mode", EUserInterfaceActionType::RadioButton, FInputGesture());
