Updated text displayed during simulation in PhAT preview viewport to use Command key instead of Control on Mac

diff --git a/Engine/Source/Editor/PhAT/Private/PhATPreviewViewportClient.cpp b/Engine/Source/Editor/PhAT/Private/PhATPreviewViewportClient.cpp
--- a/Engine/Source/Editor/PhAT/Private/PhATPreviewViewportClient.cpp
+++ b/Engine/Source/Editor/PhAT/Private/PhATPreviewViewportClient.cpp
@@ -119,7 +119,11 @@ void FPhATEdPreviewViewportClient::DrawCanvas( FViewport& InViewport, FSceneView
 	TextItem.Text = FText::GetEmpty();
 	if (SharedData->bRunningSimulation)
 	{
+#if PLATFORM_MAC
+		TextItem.Text = NSLOCTEXT("UnrealEd", "Sim_Mac", "SIM: Command+RightMouse to interact with bodies");
+#else
 		TextItem.Text = NSLOCTEXT("UnrealEd", "Sim", "SIM: Ctrl+RightMouse to interact with bodies");
+#endif
 	}
 	else if (SharedData->bSelectionLock)
 	{
