Rename PhysicsType_Fixed and PhysicsType_Unfixed and move selected simulation into toolbar

diff --git a/Engine/Config/BaseEngine.ini b/Engine/Config/BaseEngine.ini
--- a/Engine/Config/BaseEngine.ini
+++ b/Engine/Config/BaseEngine.ini
@@ -472,6 +472,9 @@ bUseStreamingPause=false
 ;EKinematicBonesUpdateToPhysics rename
 +EnumRedirects=(EnumName="EKinematicBonesUpdateToPhysics", OldEnumEntry="EKinematicBonesUpdateToPhysics::SkipFixedAndSimulatingBones",NewEnumEntry="EKinematicBonesUpdateToPhysics::SkipAllBones")
 
+;EPhysicsType rename
++EnumRedirects=(EnumName="EPhysicsType", OldEnumEntry="PhysType_Fixed",NewEnumEntry="PhysType_Kinematic")
++EnumRedirects=(EnumName="EPhysicsType", OldEnumEntry="PhysType_Unfixed",NewEnumEntry="PhysType_Simulated")
 ;Struct Redirects for UE4
 +ActiveStructRedirects=(OldStructName="VimDebugData",NewStructName="AnimBlueprintDebugData")
 
diff --git a/Engine/Source/Editor/PhAT/Private/PhAT.cpp b/Engine/Source/Editor/PhAT/Private/PhAT.cpp
--- a/Engine/Source/Editor/PhAT/Private/PhAT.cpp
+++ b/Engine/Source/Editor/PhAT/Private/PhAT.cpp
@@ -180,6 +180,7 @@ void FPhAT::InitPhAT(const EToolkitMode::Type Mode, const TSharedPtr< class IToo
 	SimulationMode = EPSM_Normal;
 	HierarchyFilterMode = PHFM_All;
 	SelectedAnimation = NULL;
+	SelectedSimulation = false;
 
 	SharedData = MakeShareable(new FPhATSharedData);
 
@@ -776,6 +777,9 @@ void FPhAT::ExtendToolbar()
 			}
 			ToolbarBuilder.EndSection();
 
+			//selected simulation
+			ToolbarBuilder.AddToolBarButton(Commands.ToggleSelectedSimulation);
+
 			//phat edit mode combo
 			FUIAction PhatMode;
 			PhatMode.CanExecuteAction = FCanExecuteAction::CreateSP(Phat, &FPhAT::IsNotSimulation);
@@ -1031,7 +1035,7 @@ void FPhAT::BindCommands()
 		Commands.ToggleSelectedSimulation,
 		FExecuteAction::CreateSP(this, &FPhAT::OnToggleSelectedSimulation),
 		FCanExecuteAction::CreateSP(this, &FPhAT::CanStartSimulation),
-		FIsActionChecked::CreateSP(this, &FPhAT::IsToggleSimulation));
+		FIsActionChecked::CreateSP(this, &FPhAT::IsSelectedSimulation));
 
 	ToolkitCommands->MapAction(
 		Commands.MeshRenderingMode_Solid,
@@ -1088,7 +1092,7 @@ void FPhAT::BindCommands()
 		FIsActionChecked::CreateSP(this, &FPhAT::IsConstraintRenderingMode, FPhATSharedData::PCV_AllLimits));
 
 	ToolkitCommands->MapAction(
-		Commands.ShowFixedBodies,
+		Commands.ShowKinematicBodies,
 		FExecuteAction::CreateSP(this, &FPhAT::OnShowFixedBodies),
 		FCanExecuteAction(),
 		FIsActionChecked::CreateSP(this, &FPhAT::IsShowFixedBodies));
@@ -1210,16 +1214,16 @@ void FPhAT::BindCommands()
 		FIsActionChecked::CreateSP(this, &FPhAT::IsShowSkeleton));
 
 	ToolkitCommands->MapAction(
-		Commands.MakeBodyFixed,
-		FExecuteAction::CreateSP(this, &FPhAT::OnSetBodyPhysicsType, EPhysicsType::PhysType_Fixed ),
+		Commands.MakeBodyKinematic,
+		FExecuteAction::CreateSP(this, &FPhAT::OnSetBodyPhysicsType, EPhysicsType::PhysType_Kinematic ),
 		FCanExecuteAction(),
-		FIsActionChecked::CreateSP(this, &FPhAT::IsBodyPhysicsType, EPhysicsType::PhysType_Fixed ) );
+		FIsActionChecked::CreateSP(this, &FPhAT::IsBodyPhysicsType, EPhysicsType::PhysType_Kinematic ) );
 
 	ToolkitCommands->MapAction(
-		Commands.MakeBodyUnfixed,
-		FExecuteAction::CreateSP(this, &FPhAT::OnSetBodyPhysicsType, EPhysicsType::PhysType_Unfixed ),
+		Commands.MakeBodySimulated,
+		FExecuteAction::CreateSP(this, &FPhAT::OnSetBodyPhysicsType, EPhysicsType::PhysType_Simulated ),
 		FCanExecuteAction(),
-		FIsActionChecked::CreateSP(this, &FPhAT::IsBodyPhysicsType, EPhysicsType::PhysType_Unfixed ) );
+		FIsActionChecked::CreateSP(this, &FPhAT::IsBodyPhysicsType, EPhysicsType::PhysType_Simulated ) );
 
 	ToolkitCommands->MapAction(
 		Commands.MakeBodyDefault,
@@ -1228,12 +1232,12 @@ void FPhAT::BindCommands()
 		FIsActionChecked::CreateSP(this, &FPhAT::IsBodyPhysicsType, EPhysicsType::PhysType_Default ) );
 
 	ToolkitCommands->MapAction(
-		Commands.FixAllBodiesBelow,
-		FExecuteAction::CreateSP(this, &FPhAT::SetBodiesBelowSelectedPhysicsType, EPhysicsType::PhysType_Fixed) );
+		Commands.KinematicAllBodiesBelow,
+		FExecuteAction::CreateSP(this, &FPhAT::SetBodiesBelowSelectedPhysicsType, EPhysicsType::PhysType_Kinematic) );
 
 	ToolkitCommands->MapAction(
-		Commands.UnfixAllBodiesBelow,
-		FExecuteAction::CreateSP(this, &FPhAT::SetBodiesBelowSelectedPhysicsType, EPhysicsType::PhysType_Unfixed) );
+		Commands.SimulatedAllBodiesBelow,
+		FExecuteAction::CreateSP(this, &FPhAT::SetBodiesBelowSelectedPhysicsType, EPhysicsType::PhysType_Simulated) );
 
 	ToolkitCommands->MapAction(
 		Commands.MakeAllBodiesBelowDefault,
@@ -1592,11 +1596,6 @@ TSharedPtr<SWidget> FPhAT::BuildMenuWidgetBody(bool bHierarchy /*= false*/)
 	{
 		const FPhATCommands& Commands = FPhATCommands::Get();
 
-		MenuBuilder.BeginSection( "PreviewTypeActions", LOCTEXT("BodyPreviewTypeHeader", "Preview" ) );
-		MenuBuilder.AddMenuEntry( Commands.ToggleSelectedSimulation );
-		MenuBuilder.EndSection();
-
-
 		struct FLocal
 		{
 			static void FillPhysicsTypeMenu(FMenuBuilder& InMenuBuilder, bool bInHierarchy)
@@ -1604,8 +1603,8 @@ TSharedPtr<SWidget> FPhAT::BuildMenuWidgetBody(bool bHierarchy /*= false*/)
 				const FPhATCommands& PhATCommands = FPhATCommands::Get();
 
 				InMenuBuilder.BeginSection("BodyPhysicsTypeActions", LOCTEXT("BodyPhysicsTypeHeader", "Body Physics Type"));
-				InMenuBuilder.AddMenuEntry(PhATCommands.MakeBodyFixed);
-				InMenuBuilder.AddMenuEntry(PhATCommands.MakeBodyUnfixed);
+				InMenuBuilder.AddMenuEntry(PhATCommands.MakeBodyKinematic);
+				InMenuBuilder.AddMenuEntry(PhATCommands.MakeBodySimulated);
 				InMenuBuilder.AddMenuEntry(PhATCommands.MakeBodyDefault);
 				InMenuBuilder.EndSection();
 				InMenuBuilder.EndSection();
@@ -1613,8 +1612,8 @@ TSharedPtr<SWidget> FPhAT::BuildMenuWidgetBody(bool bHierarchy /*= false*/)
 				if (bInHierarchy)
 				{
 					InMenuBuilder.BeginSection("BodiesBelowPhysicsTypeActions", LOCTEXT("BodiesBelowPhysicsTypeHeader", "Bodies Below Physics Type"));
-					InMenuBuilder.AddMenuEntry(PhATCommands.FixAllBodiesBelow);
-					InMenuBuilder.AddMenuEntry(PhATCommands.UnfixAllBodiesBelow);
+					InMenuBuilder.AddMenuEntry(PhATCommands.KinematicAllBodiesBelow);
+					InMenuBuilder.AddMenuEntry(PhATCommands.SimulatedAllBodiesBelow);
 					InMenuBuilder.AddMenuEntry(PhATCommands.MakeAllBodiesBelowDefault);
 					InMenuBuilder.EndSection();
 				}
@@ -1683,10 +1682,6 @@ TSharedPtr<SWidget> FPhAT::BuildMenuWidgetConstraint(bool bHierarchy /*= false*/
 	const FPhATCommands& Commands = FPhATCommands::Get();
 	FMenuBuilder MenuBuilder(bShouldCloseWindowAfterMenuSelection, GetToolkitCommands());
 	{
-		MenuBuilder.BeginSection( "PreviewTypeActions", LOCTEXT("ConstraintPreviewTypeHeader", "Preview" ) );
-		MenuBuilder.AddMenuEntry( Commands.ToggleSelectedSimulation );
-		MenuBuilder.EndSection();
-
 		MenuBuilder.BeginSection( "MotorTypeActions", LOCTEXT("ConstraintMotorTypeHeader", "Motors" ) );
 		MenuBuilder.AddMenuEntry(Commands.ToggleMotor);
 		if(bHierarchy)
@@ -2331,13 +2326,26 @@ bool FPhAT::IsSimulationMode(EPhATSimulationMode Mode) const
 void FPhAT::OnToggleSimulation()
 {
 	FixPhysicsState();
+	if (IsSelectedSimulation())
+	{
+		OnSelectedSimulation();
+	}
 	SharedData->bNoGravitySimulation = IsSimulationMode(EPSM_Gravity);
 	ImpToggleSimulation();
 }
 
+bool FPhAT::IsSelectedSimulation()
+{
+	return SelectedSimulation;
+}
+
 void FPhAT::OnToggleSelectedSimulation()
 {
-	FixPhysicsState();
+	SelectedSimulation = !SelectedSimulation;
+}
+
+void FPhAT::OnSelectedSimulation()
+{
 
 	//Before starting we modify the PhysicsType so that selected are unfixed and the rest are fixed
 	if(SharedData->bRunningSimulation == false)
@@ -2350,14 +2358,14 @@ void FPhAT::OnToggleSelectedSimulation()
 		for(int32 i=0; i<SharedData->PhysicsAsset->BodySetup.Num(); ++i)
 		{
 			BodySetup[i]->Modify();
-			BodySetup[i]->PhysicsType = PhysType_Fixed;
+			BodySetup[i]->PhysicsType = PhysType_Kinematic;
 		}
 
 		//Find which bodes have been selected
 		if(SharedData->EditingMode == FPhATSharedData::PEM_BodyEdit)
 		{
 			//Bodies already have a function that does this
-			SetBodiesBelowSelectedPhysicsType(PhysType_Unfixed);
+			SetBodiesBelowSelectedPhysicsType(PhysType_Simulated);
 		}else
 		{
 			//constraints need some more work
@@ -2377,13 +2385,10 @@ void FPhAT::OnToggleSelectedSimulation()
 				}
 			}
 
-			SetBodiesBelowPhysicsType(PhysType_Unfixed, BodyIndices);
+			SetBodiesBelowPhysicsType(PhysType_Simulated, BodyIndices);
 		}
 	}
 
-
-	SharedData->bNoGravitySimulation = IsSimulationMode(EPSM_Gravity);
-	ImpToggleSimulation();
 }
 
 
diff --git a/Engine/Source/Editor/PhAT/Private/PhAT.h b/Engine/Source/Editor/PhAT/Private/PhAT.h
--- a/Engine/Source/Editor/PhAT/Private/PhAT.h
+++ b/Engine/Source/Editor/PhAT/Private/PhAT.h
@@ -172,6 +172,8 @@ private:
 	void OnToggleSimulation();
 	void OnToggleSimulationNoGravity();
 	void OnToggleSelectedSimulation();
+	void OnSelectedSimulation();
+	bool IsSelectedSimulation();
 	bool IsToggleSimulation() const;
 	void OnMeshRenderingMode(FPhATSharedData::EPhATRenderMode Mode);
 	bool IsMeshRenderingMode(FPhATSharedData::EPhATRenderMode Mode) const;
@@ -275,6 +277,9 @@ private:
 	/** True if in OnTreeSelectionChanged()... protects against infinite recursion */
 	bool InsideSelChanged;
 
+	/** True if we want to only simulate from selected body/constraint down*/
+	bool SelectedSimulation;
+
 	/** Determines which simulation mode wer're currently in */
 	EPhATSimulationMode SimulationMode;
 
diff --git a/Engine/Source/Editor/PhAT/Private/PhATActions.cpp b/Engine/Source/Editor/PhAT/Private/PhATActions.cpp
--- a/Engine/Source/Editor/PhAT/Private/PhATActions.cpp
+++ b/Engine/Source/Editor/PhAT/Private/PhATActions.cpp
@@ -24,7 +24,7 @@ void FPhATCommands::RegisterCommands()
 	UI_COMMAND(RepeatLastSimulation, "Simulate", "Previews Physics Simulation", EUserInterfaceActionType::RadioButton, FInputGesture(EKeys::Enter));
 	UI_COMMAND(SimulationNormal, "Real Simulation", "Previews Normal Physics Simulation", EUserInterfaceActionType::RadioButton, FInputGesture());
 	UI_COMMAND(SimulationNoGravity, "No Gravity Simulation", "Run Physics Simulation without gravity. Use this to debug issues with your ragdoll. If the setup is correct, the asset should not move!", EUserInterfaceActionType::RadioButton, FInputGesture());
-	UI_COMMAND(ToggleSelectedSimulation, "Selected Simulation", "Run Physics Simulation on selected objects. Use this to tune  specific parts of your ragdoll.", EUserInterfaceActionType::Button, FInputGesture());
+	UI_COMMAND(ToggleSelectedSimulation, "Selected Simulation", "Run Physics Simulation on selected objects. Use this to tune  specific parts of your ragdoll.", EUserInterfaceActionType::ToggleButton, FInputGesture());
 	UI_COMMAND(MeshRenderingMode_Solid, "Solid", "Solid Mesh Rendering Mode", EUserInterfaceActionType::RadioButton, FInputGesture());
 	UI_COMMAND(MeshRenderingMode_Wireframe, "Wireframe", "Wireframe Mesh Rendering Mode", EUserInterfaceActionType::RadioButton, FInputGesture());
 	UI_COMMAND(MeshRenderingMode_None, "None", "No Mesh Rendering Mode", EUserInterfaceActionType::RadioButton, FInputGesture());
@@ -34,7 +34,7 @@ void FPhATCommands::RegisterCommands()
 	UI_COMMAND(ConstraintRenderingMode_None, "None", "No Constraint Rendering Mode", EUserInterfaceActionType::RadioButton, FInputGesture());
 	UI_COMMAND(ConstraintRenderingMode_AllPositions, "All Positions", "All Positions Constraint Rendering Mode", EUserInterfaceActionType::RadioButton, FInputGesture());
 	UI_COMMAND(ConstraintRenderingMode_AllLimits, "All Limits", "All Limits Constraint Rendering Mode", EUserInterfaceActionType::RadioButton, FInputGesture());
-	UI_COMMAND(ShowFixedBodies, "Fixed Bodies", "Displays Fixed Bodies In Red", EUserInterfaceActionType::ToggleButton, FInputGesture());
+	UI_COMMAND(ShowKinematicBodies, "Kinematic Bodies", "Displays Kinematic Bodies In Red", EUserInterfaceActionType::ToggleButton, FInputGesture());
 	UI_COMMAND(DrawGroundBox, "Ground Box", "Displays Floor Grid", EUserInterfaceActionType::ToggleButton, FInputGesture());
 	UI_COMMAND(ToggleGraphicsHierarchy, "Hierarchy", "Show Graphical Hierarchy Of Joints In Preview Viewport", EUserInterfaceActionType::ToggleButton, FInputGesture());
 	UI_COMMAND(ToggleBoneInfuences, "Bone Influences", "Displays Vertices Weighted To Currently Selected Bone or Body", EUserInterfaceActionType::ToggleButton, FInputGesture());
@@ -57,11 +57,11 @@ void FPhATCommands::RegisterCommands()
 	UI_COMMAND(DeleteConstraint, "Delete", "Delete Selected Constraint", EUserInterfaceActionType::Button, FInputGesture());
 	UI_COMMAND(PlayAnimation, "Play", "Play Animation", EUserInterfaceActionType::ToggleButton, FInputGesture());
 	UI_COMMAND(ShowSkeleton, "Skeleton", "Show Skeleton", EUserInterfaceActionType::ToggleButton, FInputGesture());
-	UI_COMMAND(MakeBodyFixed, "Fixed", "Fix This Body", EUserInterfaceActionType::RadioButton, FInputGesture());
-	UI_COMMAND(MakeBodyUnfixed, "Unfixed", "Unfix This Body", EUserInterfaceActionType::RadioButton, FInputGesture());
+	UI_COMMAND(MakeBodyKinematic, "Kinematic", "Make Body Kinematic", EUserInterfaceActionType::RadioButton, FInputGesture());
+	UI_COMMAND(MakeBodySimulated, "Simulated", "Make Body Simulated", EUserInterfaceActionType::RadioButton, FInputGesture());
 	UI_COMMAND(MakeBodyDefault, "Default", "Reset This Body To Default", EUserInterfaceActionType::RadioButton, FInputGesture());
-	UI_COMMAND(FixAllBodiesBelow, "Fix All Bodies Below", "Fix All Bodies Below", EUserInterfaceActionType::Button, FInputGesture());
-	UI_COMMAND(UnfixAllBodiesBelow, "Unfix All Bodies Below", "Unfix All Bodies Below", EUserInterfaceActionType::Button, FInputGesture());
+	UI_COMMAND(KinematicAllBodiesBelow, "Set All Bodies Below To Kinematic", "Set All Bodies Below To Kinematic", EUserInterfaceActionType::Button, FInputGesture());
+	UI_COMMAND(SimulatedAllBodiesBelow, "Set All Bodies Below To Simulated", "Set All Bodies Below To Simulated", EUserInterfaceActionType::Button, FInputGesture());
 	UI_COMMAND(MakeAllBodiesBelowDefault, "Reset All Bodies Below To Default", "Reset All Bodies Below To Default", EUserInterfaceActionType::Button, FInputGesture());
 	UI_COMMAND(DeleteBody, "Delete", "Delete Selected Body", EUserInterfaceActionType::Button, FInputGesture());
 	UI_COMMAND(DeleteAllBodiesBelow, "Delete All Bodies Below", "Delete All Bodies Below", EUserInterfaceActionType::Button, FInputGesture());
diff --git a/Engine/Source/Editor/PhAT/Private/PhATActions.h b/Engine/Source/Editor/PhAT/Private/PhATActions.h
--- a/Engine/Source/Editor/PhAT/Private/PhATActions.h
+++ b/Engine/Source/Editor/PhAT/Private/PhATActions.h
@@ -46,7 +46,7 @@ public:
 	TSharedPtr<FUICommandInfo> ConstraintRenderingMode_None;
 	TSharedPtr<FUICommandInfo> ConstraintRenderingMode_AllPositions;
 	TSharedPtr<FUICommandInfo> ConstraintRenderingMode_AllLimits;
-	TSharedPtr<FUICommandInfo> ShowFixedBodies;
+	TSharedPtr<FUICommandInfo> ShowKinematicBodies;
 	TSharedPtr<FUICommandInfo> DrawGroundBox;
 	TSharedPtr<FUICommandInfo> ToggleGraphicsHierarchy;
 	TSharedPtr<FUICommandInfo> ToggleBoneInfuences;
@@ -69,11 +69,11 @@ public:
 	TSharedPtr<FUICommandInfo> DeleteConstraint;
 	TSharedPtr<FUICommandInfo> PlayAnimation;
 	TSharedPtr<FUICommandInfo> ShowSkeleton;
-	TSharedPtr<FUICommandInfo> MakeBodyFixed;
-	TSharedPtr<FUICommandInfo> MakeBodyUnfixed;
+	TSharedPtr<FUICommandInfo> MakeBodyKinematic;
+	TSharedPtr<FUICommandInfo> MakeBodySimulated;
 	TSharedPtr<FUICommandInfo> MakeBodyDefault;
-	TSharedPtr<FUICommandInfo> FixAllBodiesBelow;
-	TSharedPtr<FUICommandInfo> UnfixAllBodiesBelow;
+	TSharedPtr<FUICommandInfo> KinematicAllBodiesBelow;
+	TSharedPtr<FUICommandInfo> SimulatedAllBodiesBelow;
 	TSharedPtr<FUICommandInfo> MakeAllBodiesBelowDefault;
 	TSharedPtr<FUICommandInfo> DeleteBody;
 	TSharedPtr<FUICommandInfo> DeleteAllBodiesBelow;
diff --git a/Engine/Source/Editor/PhAT/Private/PhATRender.cpp b/Engine/Source/Editor/PhAT/Private/PhATRender.cpp
--- a/Engine/Source/Editor/PhAT/Private/PhATRender.cpp
+++ b/Engine/Source/Editor/PhAT/Private/PhATRender.cpp
@@ -436,7 +436,7 @@ FColor UPhATEdSkeletalMeshComponent::GetPrimitiveColor(int32 BodyIndex, EKCollis
 	if (SharedData->bRunningSimulation)
 	{
 		// @todo draw fixed, too?
-		if (SharedData->bShowFixedStatus && BodySetup->PhysicsType == PhysType_Unfixed)
+		if (SharedData->bShowFixedStatus && BodySetup->PhysicsType == PhysType_Simulated)
 		{
 			return FixedColor;
 		}
@@ -474,7 +474,7 @@ FColor UPhATEdSkeletalMeshComponent::GetPrimitiveColor(int32 BodyIndex, EKCollis
 	
 	if (SharedData->bShowFixedStatus)
 	{
-		if (BodySetup->PhysicsType == PhysType_Unfixed)
+		if (BodySetup->PhysicsType == PhysType_Simulated)
 		{
 			return FixedColor;
 		}
diff --git a/Engine/Source/Editor/PhAT/Private/SPhATPreviewToolbar.cpp b/Engine/Source/Editor/PhAT/Private/SPhATPreviewToolbar.cpp
--- a/Engine/Source/Editor/PhAT/Private/SPhATPreviewToolbar.cpp
+++ b/Engine/Source/Editor/PhAT/Private/SPhATPreviewToolbar.cpp
@@ -57,7 +57,7 @@ TSharedRef<SWidget> SPhATPreviewViewportToolBar::GenerateViewMenu() const
 		ViewMenuBuilder.AddMenuEntry(Actions.ShowSkeleton);
 		ViewMenuBuilder.AddMenuEntry(Actions.DrawGroundBox);
 		ViewMenuBuilder.AddMenuEntry(Actions.InstanceProperties);
-		ViewMenuBuilder.AddMenuEntry(Actions.ShowFixedBodies);
+		ViewMenuBuilder.AddMenuEntry(Actions.ShowKinematicBodies);
 		ViewMenuBuilder.AddMenuEntry(Actions.ToggleGraphicsHierarchy);
 		ViewMenuBuilder.AddMenuEntry(Actions.ToggleBoneInfuences);
 		ViewMenuBuilder.AddMenuEntry(Actions.ToggleMassProperties);
diff --git a/Engine/Source/Runtime/Engine/Classes/PhysicsEngine/BodySetup.h b/Engine/Source/Runtime/Engine/Classes/PhysicsEngine/BodySetup.h
--- a/Engine/Source/Runtime/Engine/Classes/PhysicsEngine/BodySetup.h
+++ b/Engine/Source/Runtime/Engine/Classes/PhysicsEngine/BodySetup.h
@@ -30,10 +30,10 @@ enum EPhysicsType
 {
 	// follow owner option
 	PhysType_Default UMETA(DisplayName="Default"),	
-	// Do not follow owner, but fix 
-	PhysType_Fixed	UMETA(DisplayName="Fixed"),		
-	// Do not follow owner, but unfix 
-	PhysType_Unfixed	UMETA(DisplayName="Unfixed")	
+	// Do not follow owner, but make kinematic
+	PhysType_Kinematic	UMETA(DisplayName="Kinematic"),		
+	// Do not follow owner, but simulate
+	PhysType_Simulated	UMETA(DisplayName="Simulated")	
 };
 
 UENUM()
diff --git a/Engine/Source/Runtime/Engine/Private/DestructibleComponent.cpp b/Engine/Source/Runtime/Engine/Private/DestructibleComponent.cpp
--- a/Engine/Source/Runtime/Engine/Private/DestructibleComponent.cpp
+++ b/Engine/Source/Runtime/Engine/Private/DestructibleComponent.cpp
@@ -1128,7 +1128,7 @@ FBodyInstance* UDestructibleComponent::GetBodyInstance( FName BoneName /*= NAME_
 
 bool UDestructibleComponent::IsAnySimulatingPhysics() const 
 {
-	return BodyInstance.bSimulatePhysics;
+	return !!BodyInstance.bSimulatePhysics;
 }
 
 #if WITH_PHYSX
diff --git a/Engine/Source/Runtime/Engine/Private/PhysicsEngine/BodyInstance.cpp b/Engine/Source/Runtime/Engine/Private/PhysicsEngine/BodyInstance.cpp
--- a/Engine/Source/Runtime/Engine/Private/PhysicsEngine/BodyInstance.cpp
+++ b/Engine/Source/Runtime/Engine/Private/PhysicsEngine/BodyInstance.cpp
@@ -754,15 +754,15 @@ void FBodyInstance::InitBody(UBodySetup* Setup, const FTransform& Transform, UPr
 		SkelMeshComp = Cast<USkeletalMeshComponent>(OwnerComponent.Get());
 
 		// if not dedicated and if set up is done to use physics, enable flag on
-		if (BodySetup->PhysicsType == PhysType_Unfixed || BodySetup->PhysicsType == PhysType_Default)
+		if (BodySetup->PhysicsType == PhysType_Simulated || BodySetup->PhysicsType == PhysType_Default)
 		{
 			bool bEnableSimulation = (SkelMeshComp && IsRunningDedicatedServer())? SkelMeshComp->bEnablePhysicsOnDedicatedServer: true;
-			bEnableSimulation &= ((BodySetup->PhysicsType == PhysType_Unfixed) || (SkelMeshComp->BodyInstance.bSimulatePhysics));	//if unfixed enable. If default look at parent
+			bEnableSimulation &= ((BodySetup->PhysicsType == PhysType_Simulated) || (SkelMeshComp->BodyInstance.bSimulatePhysics));	//if unfixed enable. If default look at parent
 			if ( bEnableSimulation )
 			{
 				// set simulate to true if using physics
 				bSimulatePhysics = true;
-				if (BodySetup->PhysicsType == PhysType_Unfixed)
+				if (BodySetup->PhysicsType == PhysType_Simulated)
 				{
 					PhysicsBlendWeight = 1.f;
 				}
@@ -771,7 +771,7 @@ void FBodyInstance::InitBody(UBodySetup* Setup, const FTransform& Transform, UPr
 	 		else
 	 		{
 	 			bSimulatePhysics = false;
-				if (BodySetup->PhysicsType == PhysType_Unfixed)
+				if (BodySetup->PhysicsType == PhysType_Simulated)
 				{
 					PhysicsBlendWeight = 0.f;
 				}
diff --git a/Engine/Source/Runtime/Engine/Private/PhysicsEngine/BodySetup.cpp b/Engine/Source/Runtime/Engine/Private/PhysicsEngine/BodySetup.cpp
--- a/Engine/Source/Runtime/Engine/Private/PhysicsEngine/BodySetup.cpp
+++ b/Engine/Source/Runtime/Engine/Private/PhysicsEngine/BodySetup.cpp
@@ -668,11 +668,11 @@ void UBodySetup::PostLoad()
 	{
 		if ( bAlwaysFullAnimWeight_DEPRECATED )
 		{
-			PhysicsType = PhysType_Unfixed;
+			PhysicsType = PhysType_Simulated;
 		}
 		else if ( DefaultInstance.bSimulatePhysics == false )
 		{
-			PhysicsType = PhysType_Fixed;
+			PhysicsType = PhysType_Kinematic;
 		}
 		else
 		{
diff --git a/Engine/Source/Runtime/Engine/Private/SkeletalMeshComponentPhysics.cpp b/Engine/Source/Runtime/Engine/Private/SkeletalMeshComponentPhysics.cpp
--- a/Engine/Source/Runtime/Engine/Private/SkeletalMeshComponentPhysics.cpp
+++ b/Engine/Source/Runtime/Engine/Private/SkeletalMeshComponentPhysics.cpp
@@ -994,7 +994,7 @@ void USkeletalMeshComponent::ResetAllBodiesSimulatePhysics()
 		// Set fixed on any bodies with bAlwaysFullAnimWeight set to true
 		if(BodySetup && BodySetup->PhysicsType != PhysType_Default)
 		{
-			if (BodySetup->PhysicsType == PhysType_Unfixed)
+			if (BodySetup->PhysicsType == PhysType_Simulated)
 			{
 				BodyInst->SetInstanceSimulatePhysics(true);
 			}
