Fix Phat body selection not working when pick translucent is disabled
TTP#335501

diff --git a/Engine/Source/Editor/PhAT/Classes/PhATEdSkeletalMeshComponent.h b/Engine/Source/Editor/PhAT/Classes/PhATEdSkeletalMeshComponent.h
--- a/Engine/Source/Editor/PhAT/Classes/PhATEdSkeletalMeshComponent.h
+++ b/Engine/Source/Editor/PhAT/Classes/PhATEdSkeletalMeshComponent.h
@@ -38,6 +38,8 @@ class UPhATEdSkeletalMeshComponent : public USkeletalMeshComponent
 	UPROPERTY(transient)
 	UMaterialInterface* BoneUnselectedMaterial;
 	UPROPERTY(transient)
+	UMaterialInterface* BoneMaterialHit;
+	UPROPERTY(transient)
 	UMaterialInterface* BoneNoCollisionMaterial;
 
 	/** Mesh-space matrices showing state of just animation (ie before physics) - useful for debugging! */
@@ -68,7 +70,7 @@ class UPhATEdSkeletalMeshComponent : public USkeletalMeshComponent
 	/** Accessors/helper methods */
 	FTransform GetPrimitiveTransform(FTransform& BoneTM, int32 BodyIndex, EKCollisionPrimitiveType PrimType, int32 PrimIndex, float Scale);
 	FColor GetPrimitiveColor(int32 BodyIndex, EKCollisionPrimitiveType PrimitiveType, int32 PrimitiveIndex);
-	UMaterialInterface* GetPrimitiveMaterial(int32 BodyIndex, EKCollisionPrimitiveType PrimitiveType, int32 PrimitiveIndex);
+	UMaterialInterface* GetPrimitiveMaterial(int32 BodyIndex, EKCollisionPrimitiveType PrimitiveType, int32 PrimitiveIndex, bool bHitTest);
 
 	/** Returns the physics asset for this PhATEd component - note: This hides the implementation in the USkinnedMeshComponent base class */
 	class UPhysicsAsset* GetPhysicsAsset() const;
diff --git a/Engine/Source/Editor/PhAT/Private/PhATRender.cpp b/Engine/Source/Editor/PhAT/Private/PhATRender.cpp
--- a/Engine/Source/Editor/PhAT/Private/PhATRender.cpp
+++ b/Engine/Source/Editor/PhAT/Private/PhATRender.cpp
@@ -30,6 +30,9 @@ UPhATEdSkeletalMeshComponent::UPhATEdSkeletalMeshComponent(const class FPostCons
 	BoneSelectedMaterial = LoadObject<UMaterialInterface>(NULL, TEXT("/Engine/EditorMaterials/PhAT_BoneSelectedMaterial.PhAT_BoneSelectedMaterial"), NULL, LOAD_None, NULL);
 	check(BoneSelectedMaterial);
 
+	BoneMaterialHit = UMaterial::GetDefaultMaterial(MD_Surface);
+	check(BoneMaterialHit);
+
 	BoneUnselectedMaterial = LoadObject<UMaterialInterface>(NULL, TEXT("/Engine/EditorMaterials/PhAT_UnselectedMaterial.PhAT_UnselectedMaterial"), NULL, LOAD_None, NULL);
 	check(BoneUnselectedMaterial);
 
@@ -92,7 +95,7 @@ void UPhATEdSkeletalMeshComponent::RenderAssetTools(const FSceneView* View, clas
 				//solids are drawn if it's the ViewMode and we're not doing a hit, or if it's hitAndBodyMode
 				if( (CollisionViewMode == FPhATSharedData::PRM_Solid && !bHitTest) || bHitTestAndBodyMode)
 				{
-					UMaterialInterface*	PrimMaterial = GetPrimitiveMaterial(i, KPT_Sphere, j);
+					UMaterialInterface*	PrimMaterial = GetPrimitiveMaterial(i, KPT_Sphere, j, bHitTestAndBodyMode);
 					AggGeom->SphereElems[j].DrawElemSolid(PDI, ElemTM, 1.f, PrimMaterial->GetRenderProxy(0));
 				}
 
@@ -123,7 +126,7 @@ void UPhATEdSkeletalMeshComponent::RenderAssetTools(const FSceneView* View, clas
 
 				if ( (CollisionViewMode == FPhATSharedData::PRM_Solid && !bHitTest) || bHitTestAndBodyMode)
 				{
-					UMaterialInterface*	PrimMaterial = GetPrimitiveMaterial(i, KPT_Box, j);
+					UMaterialInterface*	PrimMaterial = GetPrimitiveMaterial(i, KPT_Box, j, bHitTestAndBodyMode);
 					AggGeom->BoxElems[j].DrawElemSolid(PDI, ElemTM, 1.f, PrimMaterial->GetRenderProxy(0));
 				}
 
@@ -153,7 +156,7 @@ void UPhATEdSkeletalMeshComponent::RenderAssetTools(const FSceneView* View, clas
 
 				if ( (CollisionViewMode == FPhATSharedData::PRM_Solid && !bHitTest) || bHitTestAndBodyMode)
 				{
-					UMaterialInterface*	PrimMaterial = GetPrimitiveMaterial(i, KPT_Sphyl, j);
+					UMaterialInterface*	PrimMaterial = GetPrimitiveMaterial(i, KPT_Sphyl, j, bHitTestAndBodyMode);
 					AggGeom->SphylElems[j].DrawElemSolid(PDI, ElemTM, 1.f, PrimMaterial->GetRenderProxy(0));
 				}
 
@@ -497,11 +500,11 @@ FColor UPhATEdSkeletalMeshComponent::GetPrimitiveColor(int32 BodyIndex, EKCollis
 	}
 }
 
-UMaterialInterface* UPhATEdSkeletalMeshComponent::GetPrimitiveMaterial(int32 BodyIndex, EKCollisionPrimitiveType PrimitiveType, int32 PrimitiveIndex)
+UMaterialInterface* UPhATEdSkeletalMeshComponent::GetPrimitiveMaterial(int32 BodyIndex, EKCollisionPrimitiveType PrimitiveType, int32 PrimitiveIndex, bool bHitTest)
 {
 	if (SharedData->bRunningSimulation || SharedData->EditingMode == FPhATSharedData::PEM_ConstraintEdit)
 	{
-		return BoneUnselectedMaterial;
+		return bHitTest ? BoneMaterialHit : BoneUnselectedMaterial;
 	}
 
 	FPhATSharedData::FSelection Body(BodyIndex, PrimitiveType, PrimitiveIndex);
@@ -510,18 +513,18 @@ UMaterialInterface* UPhATEdSkeletalMeshComponent::GetPrimitiveMaterial(int32 Bod
 	{
 		if (Body == SharedData->SelectedBodies[i])
 		{
-			return ElemSelectedMaterial;
+			return bHitTest ? BoneMaterialHit : ElemSelectedMaterial;
 		}
 	}
 
 	// If there is no collision with this body, use 'no collision material'.
 	if (SharedData->NoCollisionBodies.Find(BodyIndex) != INDEX_NONE)
 	{
-		return BoneNoCollisionMaterial;
+		return bHitTest ? BoneMaterialHit : BoneNoCollisionMaterial;
 	}
 	else
 	{
-		return BoneUnselectedMaterial;
+		return bHitTest ? BoneMaterialHit : BoneUnselectedMaterial;
 	}
 
 }
