Fix crash in PhAT when renaming constraint TTP#325202

diff --git a/Engine/Source/Editor/PhAT/Private/PhATSharedData.cpp b/Engine/Source/Editor/PhAT/Private/PhATSharedData.cpp
--- a/Engine/Source/Editor/PhAT/Private/PhATSharedData.cpp
+++ b/Engine/Source/Editor/PhAT/Private/PhATSharedData.cpp
@@ -894,12 +894,13 @@ void FPhATSharedData::SetSelectedConstraintRelTM(const FTransform& RelTM)
 
 	// Get child bone transform
 	int32 BoneIndex = EditorSkelMesh->RefSkeleton.FindBoneIndex(ConstraintSetup->DefaultInstance.ConstraintBone1);
-	check(BoneIndex != INDEX_NONE);
-
-	FTransform BoneTM = EditorSkelComp->GetBoneTransform(BoneIndex);
-	BoneTM.RemoveScaling();
+	if (BoneIndex != INDEX_NONE)
+	{
+		FTransform BoneTM = EditorSkelComp->GetBoneTransform(BoneIndex);
+		BoneTM.RemoveScaling();
 
-	ConstraintSetup->DefaultInstance.SetRefFrame(EConstraintFrame::Frame1, WNewChildFrame.GetRelativeTransform(BoneTM));
+		ConstraintSetup->DefaultInstance.SetRefFrame(EConstraintFrame::Frame1, WNewChildFrame.GetRelativeTransform(BoneTM));
+	}
 }
 
 FTransform FPhATSharedData::GetConstraintWorldTM(const FSelection * Constraint, EConstraintFrame::Type Frame) const
@@ -923,12 +924,19 @@ FTransform FPhATSharedData::GetConstraintWorldTM(const FSelection * Constraint,
 	{
 		BoneIndex = EditorSkelMesh->RefSkeleton.FindBoneIndex(ConstraintSetup->DefaultInstance.ConstraintBone2);
 	}
-	check(BoneIndex != INDEX_NONE);
 
-	FTransform BoneTM = EditorSkelComp->GetBoneTransform(BoneIndex);
-	BoneTM.RemoveScaling();
+	//It's possible for BoneIndex to be INDEX_NONE, for example if the constraint name is invalid
+	if (BoneIndex != INDEX_NONE)
+	{
+		FTransform BoneTM = EditorSkelComp->GetBoneTransform(BoneIndex);
+		BoneTM.RemoveScaling();
 
-	return FrameTM * BoneTM;
+		return FrameTM * BoneTM;
+	}
+	else
+	{
+		return FTransform::Identity;
+	}	
 }
 
 void FPhATSharedData::CopyConstraint()
