Improve undo in PhAT so that non destructive undos maintain selection

diff --git a/Engine/Source/Editor/PhAT/Private/PhAT.cpp b/Engine/Source/Editor/PhAT/Private/PhAT.cpp
--- a/Engine/Source/Editor/PhAT/Private/PhAT.cpp
+++ b/Engine/Source/Editor/PhAT/Private/PhAT.cpp
@@ -586,8 +586,7 @@ void FPhAT::CycleTransformMode()
 
 void FPhAT::PostUndo(bool bSuccess)
 {
-	SharedData->SetSelectedBody(NULL);
-	SharedData->SetSelectedConstraint(INDEX_NONE);
+	SharedData->PostUndo();
 	RefreshHierachyTree();
 
 	SharedData->RefreshPhysicsAssetChange(SharedData->PhysicsAsset);
diff --git a/Engine/Source/Editor/PhAT/Private/PhATSharedData.cpp b/Engine/Source/Editor/PhAT/Private/PhATSharedData.cpp
--- a/Engine/Source/Editor/PhAT/Private/PhATSharedData.cpp
+++ b/Engine/Source/Editor/PhAT/Private/PhATSharedData.cpp
@@ -1377,19 +1377,58 @@ void FPhATSharedData::OpenNewBodyDlg(FPhysAssetCreateParams* NewBodyData, EAppRe
 	GEditor->EditorAddModalWindow(ModalWindow);
 }
 
-void FPhATSharedData::Undo()
+void FPhATSharedData::PostUndo()
 {
 	if (bRunningSimulation)
 	{
 		return;
 	}
 
-	// Clear selection before we undo. We don't transact the editor itself - don't want to have something selected that is then removed.
-	SetSelectedBody(NULL);
-	SetSelectedConstraint(INDEX_NONE);
+	bool bInvalidSelection = false;
 
-	GEditor->UndoTransaction();
-	PhysicsAsset->UpdateBodySetupIndexMap();
+	for (int32 BodyIndex = 0; BodyIndex < SelectedBodies.Num() && bInvalidSelection == false; ++BodyIndex)
+	{
+		const FSelection & Selection = SelectedBodies[BodyIndex];
+		if (PhysicsAsset->BodySetup.Num() <= Selection.Index)
+		{
+			bInvalidSelection = true;
+		}
+		else
+		{
+			
+			if (UBodySetup * BodySetup = PhysicsAsset->BodySetup[Selection.Index])
+			{
+				switch (Selection.PrimitiveType)
+				{
+				case KPT_Box: bInvalidSelection = BodySetup->AggGeom.BoxElems.Num() <= Selection.PrimitiveIndex ? true : bInvalidSelection; break;
+				case KPT_Convex: bInvalidSelection = BodySetup->AggGeom.ConvexElems.Num() <= Selection.PrimitiveIndex ? true : bInvalidSelection; break;
+				case KPT_Sphere: bInvalidSelection = BodySetup->AggGeom.SphereElems.Num() <= Selection.PrimitiveIndex ? true : bInvalidSelection; break;
+				case KPT_Sphyl: bInvalidSelection = BodySetup->AggGeom.SphylElems.Num() <= Selection.PrimitiveIndex ? true : bInvalidSelection; break;
+				default: bInvalidSelection = true;
+				}
+			}
+			else
+			{
+				bInvalidSelection = true;
+			}
+		}
+	}
+
+	for (int32 ConstraintIndex = 0; ConstraintIndex < SelectedConstraints.Num() && bInvalidSelection == false; ++ConstraintIndex)
+	{
+		const FSelection & Selection = SelectedConstraints[ConstraintIndex];
+		if (PhysicsAsset->ConstraintSetup.Num() <= Selection.Index)
+		{
+			bInvalidSelection = true;
+		}
+	}
+
+	if (bInvalidSelection)
+	{
+		// Clear selection before we undo. We don't transact the editor itself - don't want to have something selected that is then removed.
+		SetSelectedBody(NULL);
+		SetSelectedConstraint(INDEX_NONE);
+	}
 
 	PreviewChangedEvent.Broadcast();
 	HierarchyChangedEvent.Broadcast();
diff --git a/Engine/Source/Editor/PhAT/Private/PhATSharedData.h b/Engine/Source/Editor/PhAT/Private/PhATSharedData.h
--- a/Engine/Source/Editor/PhAT/Private/PhATSharedData.h
+++ b/Engine/Source/Editor/PhAT/Private/PhATSharedData.h
@@ -116,7 +116,7 @@ public:
 	void HitConstraint(int32 ConstraintIndex, bool bGroupSelect = false);
 
 	/** Undo/Redo */
-	void Undo();
+	void PostUndo();
 	void Redo();
 	void SetCollisionBetweenSelected(bool bEnableCollision);
 
