PhAT copying constraints now copies frame 2 relative to frame1

diff --git a/Engine/Source/Editor/PhAT/Private/PhATSharedData.cpp b/Engine/Source/Editor/PhAT/Private/PhATSharedData.cpp
--- a/Engine/Source/Editor/PhAT/Private/PhATSharedData.cpp
+++ b/Engine/Source/Editor/PhAT/Private/PhATSharedData.cpp
@@ -910,42 +910,6 @@ void FPhATSharedData::SetSelectedConstraintRelTM(const FTransform& RelTM)
 	}
 }
 
-FTransform FPhATSharedData::GetConstraintWorldTM(const FSelection * Constraint, EConstraintFrame::Type Frame) const
-{
-	int32 ConstraintIndex = Constraint ? Constraint->Index : INDEX_NONE;
-	if (ConstraintIndex == INDEX_NONE)
-	{
-		return FTransform::Identity;
-	}
-
-	UPhysicsConstraintTemplate* ConstraintSetup = PhysicsAsset->ConstraintSetup[ConstraintIndex];
-
-	FTransform FrameTM = ConstraintSetup->DefaultInstance.GetRefFrame(Frame);
-
-	int32 BoneIndex;
-	if (Frame == EConstraintFrame::Frame1)
-	{
-		BoneIndex = EditorSkelMesh->RefSkeleton.FindBoneIndex(ConstraintSetup->DefaultInstance.ConstraintBone1);
-	}
-	else
-	{
-		BoneIndex = EditorSkelMesh->RefSkeleton.FindBoneIndex(ConstraintSetup->DefaultInstance.ConstraintBone2);
-	}
-
-	//It's possible for BoneIndex to be INDEX_NONE, for example if the constraint name is invalid
-	if (BoneIndex != INDEX_NONE)
-	{
-		FTransform BoneTM = EditorSkelComp->GetBoneTransform(BoneIndex);
-		BoneTM.RemoveScaling();
-
-		return FrameTM * BoneTM;
-	}
-	else
-	{
-		return FTransform::Identity;
-	}	
-}
-
 void FPhATSharedData::CopyConstraint()
 {
 	check(SelectedConstraints.Num() == 1);
@@ -962,12 +926,35 @@ void FPhATSharedData::PasteConstraintProperties()
 
 	const FScopedTransaction Transaction( NSLOCTEXT("PhAT", "PasteConstraintProperties", "Paste Constraint Properties") );
 
+	UPhysicsConstraintTemplate* FromConstraintSetup = CopiedConstraintTemplate;
+
+	//We want to copy frame2 relative to frame1. To do this we need to go from body2 into body1, compute relative.
+	//Then apply relative inside body1 of the destination constraint, and finally move it into body2 of destination constraint
+	
+	//In total there are 4 bodies, body a and b are the frame1 frame2 bodies of the copied
+	// body c and d are frame 1 and frame 2 of the destination constraint
+	//we use the body letter to denote which space we're in, for example Frame1a would be frame1 inside BodyA
+	//Frame1b would be frame 1 inside body b. In this case BodyA * Frame1a == BodyB * Frame1b
+
+	FTransform Frame1A = FromConstraintSetup->DefaultInstance.GetRefFrame(EConstraintFrame::Frame1);
+	FTransform BodyA = GetConstraintBodyTM(FromConstraintSetup, EConstraintFrame::Frame1);
+	FTransform BodyB = GetConstraintBodyTM(FromConstraintSetup, EConstraintFrame::Frame2);
+	FTransform BodyAFrame1A = GetConstraintWorldTM(FromConstraintSetup, EConstraintFrame::Frame1);
+	FTransform BodyBFrame2B = GetConstraintWorldTM(FromConstraintSetup, EConstraintFrame::Frame2);
 
 	for(int32 i=0; i<SelectedConstraints.Num(); ++i)
 	{
 		// If we are showing instance properties - copy instance properties. If showing setup, just copy setup properties.
 		UPhysicsConstraintTemplate* ToConstraintSetup = PhysicsAsset->ConstraintSetup[SelectedConstraints[i].Index];
-		UPhysicsConstraintTemplate* FromConstraintSetup = CopiedConstraintTemplate;
+
+		FTransform Frame1C = ToConstraintSetup->DefaultInstance.GetRefFrame(EConstraintFrame::Frame1);
+		FTransform BodyC = GetConstraintBodyTM(ToConstraintSetup, EConstraintFrame::Frame1);
+		FTransform BodyD = GetConstraintBodyTM(ToConstraintSetup, EConstraintFrame::Frame2);
+
+		FTransform FromF1AToF2A = BodyBFrame2B * BodyA.InverseSafe() * Frame1A.InverseSafe();
+		FTransform Frame2C = FromF1AToF2A * Frame1C;
+		FTransform Frame2D = Frame2C* BodyC * BodyD.InverseSafe();
+
 
 		ToConstraintSetup->Modify();
 		FConstraintInstance OldInstance = ToConstraintSetup->DefaultInstance;
@@ -982,10 +969,14 @@ void FPhATSharedData::PasteConstraintProperties()
 		ToConstraintSetup->DefaultInstance.ConstraintBone2		= OldInstance.ConstraintBone2;
 		ToConstraintSetup->DefaultInstance.Pos1					= OldInstance.Pos1;
 		ToConstraintSetup->DefaultInstance.Pos2					= OldInstance.Pos2;
-		ToConstraintSetup->DefaultInstance.PriAxis1				= OldInstance.PriAxis1;
-		ToConstraintSetup->DefaultInstance.PriAxis2				= OldInstance.PriAxis2;
-		ToConstraintSetup->DefaultInstance.SecAxis1				= OldInstance.SecAxis1;
-		ToConstraintSetup->DefaultInstance.SecAxis2				= OldInstance.SecAxis2;
+
+		//frame1 stays the same
+		ToConstraintSetup->DefaultInstance.PriAxis1 = OldInstance.PriAxis1;
+		ToConstraintSetup->DefaultInstance.SecAxis1 = OldInstance.SecAxis1;
+
+		//frame2 is copied but relative to frame1
+		ToConstraintSetup->DefaultInstance.PriAxis2 = Frame2D.GetUnitAxis(EAxis::X);
+		ToConstraintSetup->DefaultInstance.SecAxis2 = Frame2D.GetUnitAxis(EAxis::Y);
 	}
 	
 }
@@ -1180,10 +1171,44 @@ void FPhATSharedData::DeleteCurrentPrim()
 	SetSelectedBodyAnyPrim(INDEX_NONE); // Will call UpdateViewport
 	RefreshPhysicsAssetChange(PhysicsAsset);
 }
+FTransform FPhATSharedData::GetConstraintBodyTM(const UPhysicsConstraintTemplate * ConstraintSetup, EConstraintFrame::Type Frame) const
+{
+	if (ConstraintSetup == NULL)
+	{
+		return FTransform::Identity;
+	}
+
+	int32 BoneIndex;
+	if (Frame == EConstraintFrame::Frame1)
+	{
+		BoneIndex = EditorSkelMesh->RefSkeleton.FindBoneIndex(ConstraintSetup->DefaultInstance.ConstraintBone1);
+	}
+	else
+	{
+		BoneIndex = EditorSkelMesh->RefSkeleton.FindBoneIndex(ConstraintSetup->DefaultInstance.ConstraintBone2);
+	}
+
+	// If we couldn't find the bone - fall back to identity.
+	if (BoneIndex == INDEX_NONE)
+	{
+		return FTransform::Identity;
+	}
+	else
+	{
+		FTransform BoneTM = EditorSkelComp->GetBoneTransform(BoneIndex);
+		BoneTM.RemoveScaling();
+
+		return BoneTM;
+	}
+}
 
-FTransform FPhATSharedData::GetConstraintMatrix(int32 ConstraintIndex, EConstraintFrame::Type Frame, float Scale)
+FTransform FPhATSharedData::GetConstraintWorldTM(const UPhysicsConstraintTemplate * ConstraintSetup, EConstraintFrame::Type Frame, float Scale) const
 {
-	UPhysicsConstraintTemplate* ConstraintSetup = PhysicsAsset->ConstraintSetup[ConstraintIndex];
+	if (ConstraintSetup == NULL)
+	{
+		return FTransform::Identity;
+	}
+
 	FVector Scale3D(Scale);
 
 	int32 BoneIndex;
@@ -1213,6 +1238,26 @@ FTransform FPhATSharedData::GetConstraintMatrix(int32 ConstraintIndex, EConstrai
 	}
 }
 
+FTransform FPhATSharedData::GetConstraintMatrix(int32 ConstraintIndex, EConstraintFrame::Type Frame, float Scale) const
+{
+	UPhysicsConstraintTemplate* ConstraintSetup = PhysicsAsset->ConstraintSetup[ConstraintIndex];
+	return GetConstraintWorldTM(ConstraintSetup, Frame, Scale);
+}
+
+
+FTransform FPhATSharedData::GetConstraintWorldTM(const FSelection * Constraint, EConstraintFrame::Type Frame) const
+{
+	int32 ConstraintIndex = Constraint ? Constraint->Index : INDEX_NONE;
+	if (ConstraintIndex == INDEX_NONE)
+	{
+		return FTransform::Identity;
+	}
+
+	UPhysicsConstraintTemplate* ConstraintSetup = PhysicsAsset->ConstraintSetup[ConstraintIndex];
+	return GetConstraintWorldTM(ConstraintSetup, Frame, 1.f);
+}
+
+
 void FPhATSharedData::DeleteCurrentConstraint()
 {
 	if (EditingMode != PEM_ConstraintEdit || !GetSelectedConstraint())
diff --git a/Engine/Source/Editor/PhAT/Private/PhATSharedData.h b/Engine/Source/Editor/PhAT/Private/PhATSharedData.h
--- a/Engine/Source/Editor/PhAT/Private/PhATSharedData.h
+++ b/Engine/Source/Editor/PhAT/Private/PhATSharedData.h
@@ -80,8 +80,13 @@ public:
 	/** Constraint editing */
 	void SetSelectedConstraint(int32 ConstraintIndex, bool bGroupSelect = false);
 	FTransform GetConstraintWorldTM(const FSelection * Constraint, EConstraintFrame::Type Frame) const;
+	FTransform GetConstraintWorldTM(const UPhysicsConstraintTemplate * ConstraintSetup, EConstraintFrame::Type Frame, float Scale = 1.f) const;
+	FTransform GetConstraintMatrix(int32 ConstraintIndex, EConstraintFrame::Type Frame, float Scale) const;
+	
+	FTransform GetConstraintBodyTM(const UPhysicsConstraintTemplate * ConstraintSetup, EConstraintFrame::Type Frame) const;
+
 	void SetSelectedConstraintRelTM(const FTransform& RelTM);	
-	FTransform GetConstraintMatrix(int32 ConstraintIndex, EConstraintFrame::Type Frame, float Scale);
+	
 	void DeleteCurrentConstraint();
 	void PasteConstraintProperties();
 	
