Fix PhAT adding new bodies not correctly creating joints for non-leaf bodies

diff --git a/Engine/Source/Editor/PhAT/Private/PhATSharedData.cpp b/Engine/Source/Editor/PhAT/Private/PhATSharedData.cpp
--- a/Engine/Source/Editor/PhAT/Private/PhATSharedData.cpp
+++ b/Engine/Source/Editor/PhAT/Private/PhATSharedData.cpp
@@ -868,7 +868,7 @@ void FPhATSharedData::MakeNewBody(int32 NewBoneIndex)
 				UPhysicsConstraintTemplate* ChildConstraintSetup = PhysicsAsset->ConstraintSetup[ ConstraintIndex ];
 				check(ChildConstraintSetup);
 
-				InitConstraintSetup(ChildConstraintSetup, NewBodyIndex, ChildBodyIndex);
+				InitConstraintSetup(ChildConstraintSetup, ChildBodyIndex, NewBodyIndex);
 			}
 		}
 	}
