Fix crash when deleting bodies in phat that are not associated to any bone

diff --git a/Engine/Source/Editor/PhAT/Private/PhATSharedData.cpp b/Engine/Source/Editor/PhAT/Private/PhATSharedData.cpp
--- a/Engine/Source/Editor/PhAT/Private/PhATSharedData.cpp
+++ b/Engine/Source/Editor/PhAT/Private/PhATSharedData.cpp
@@ -1073,30 +1073,33 @@ void FPhATSharedData::DeleteBody(int32 DelBodyIndex, bool bRefreshComponent)
 	PhysicsAsset->GetNearestBodyIndicesBelow(NearestBodiesBelow, BodySetup->BoneName, EditorSkelMesh);
 	
 	int32 BoneIndex = EditorSkelMesh->RefSkeleton.FindBoneIndex(BodySetup->BoneName);
-	check(BoneIndex != INDEX_NONE);
-	int32 ParentBodyIndex = PhysicsAsset->FindParentBodyIndex(EditorSkelMesh, BoneIndex);
 
-	UBodySetup * ParentBody = ParentBodyIndex != INDEX_NONE ? PhysicsAsset->BodySetup[ParentBodyIndex] : NULL;
-
-	for (int32 i = 0; i <Constraints.Num(); ++i)
+	if (BoneIndex != INDEX_NONE)	//it's possible to delete bodies that have no bones. In this case just ignore all of this fixup code
 	{
-		int32 ConstraintIndex = Constraints[i];
-		UPhysicsConstraintTemplate * Constraint = PhysicsAsset->ConstraintSetup[ConstraintIndex];
-		Constraint->Modify();
+		int32 ParentBodyIndex = PhysicsAsset->FindParentBodyIndex(EditorSkelMesh, BoneIndex);
+
+		UBodySetup * ParentBody = ParentBodyIndex != INDEX_NONE ? PhysicsAsset->BodySetup[ParentBodyIndex] : NULL;
 
-		if (ParentBody)
+		for (int32 i = 0; i < Constraints.Num(); ++i)
 		{
-			//for all constraints that contain a nearest child of this body, create a copy of the constraint between the child and parent
-			for (int32 i = 0; i < NearestBodiesBelow.Num(); ++i)
-			{
-				int32 BodyBelowIndex = NearestBodiesBelow[i];
-				UBodySetup * BodyBelow = PhysicsAsset->BodySetup[BodyBelowIndex];
+			int32 ConstraintIndex = Constraints[i];
+			UPhysicsConstraintTemplate * Constraint = PhysicsAsset->ConstraintSetup[ConstraintIndex];
+			Constraint->Modify();
 
-				if (Constraint->DefaultInstance.ConstraintBone1 == BodyBelow->BoneName)
+			if (ParentBody)
+			{
+				//for all constraints that contain a nearest child of this body, create a copy of the constraint between the child and parent
+				for (int32 i = 0; i < NearestBodiesBelow.Num(); ++i)
 				{
-					int32 NewConstraintIndex = FPhysicsAssetUtils::CreateNewConstraint(PhysicsAsset, BodyBelow->BoneName, Constraint);
-					UPhysicsConstraintTemplate * NewConstraint = PhysicsAsset->ConstraintSetup[NewConstraintIndex];
-					InitConstraintSetup(NewConstraint, BodyBelowIndex, ParentBodyIndex);
+					int32 BodyBelowIndex = NearestBodiesBelow[i];
+					UBodySetup * BodyBelow = PhysicsAsset->BodySetup[BodyBelowIndex];
+
+					if (Constraint->DefaultInstance.ConstraintBone1 == BodyBelow->BoneName)
+					{
+						int32 NewConstraintIndex = FPhysicsAssetUtils::CreateNewConstraint(PhysicsAsset, BodyBelow->BoneName, Constraint);
+						UPhysicsConstraintTemplate * NewConstraint = PhysicsAsset->ConstraintSetup[NewConstraintIndex];
+						InitConstraintSetup(NewConstraint, BodyBelowIndex, ParentBodyIndex);
+					}
 				}
 			}
 		}
