Analytics addition for Persona/Phat/CreateAsset

diff --git a/Engine/Source/Developer/AssetTools/Private/AssetTools.cpp b/Engine/Source/Developer/AssetTools/Private/AssetTools.cpp
--- a/Engine/Source/Developer/AssetTools/Private/AssetTools.cpp
+++ b/Engine/Source/Developer/AssetTools/Private/AssetTools.cpp
@@ -252,6 +252,9 @@ UObject* FAssetTools::CreateAsset(const FString& AssetName, const FString& Packa
 		// Notify the asset registry
 		FAssetRegistryModule::AssetCreated(NewObj);
 
+		// analytics create record
+		FAssetTools::OnNewCreateRecord(AssetClass, false);
+
 		// Mark the package dirty...
 		Pkg->MarkPackageDirty();
 	}
@@ -297,6 +300,9 @@ UObject* FAssetTools::DuplicateAsset(const FString& AssetName, const FString& Pa
 		SourceControlHelpers::BranchPackage(NewObject->GetOutermost(), OriginalObject->GetOutermost());
 	}
 
+	// analytics create record
+	FAssetTools::OnNewCreateRecord(NewObject->GetClass(), true);
+
 	return NewObject;
 }
 
@@ -1050,7 +1056,6 @@ void FAssetTools::MigratePackages(const TArray<FName>& PackageNamesToMigrate) co
 	}
 }
 
-
 void FAssetTools::OnNewImportRecord(UClass* AssetType, const FString& FileExtension, bool bSucceeded, bool bWasCancelled, const FDateTime& StartTime)
 {
 	// Don't attempt to report usage stats if analytics isn't available
@@ -1067,6 +1072,19 @@ void FAssetTools::OnNewImportRecord(UClass* AssetType, const FString& FileExtens
 	}
 }
 
+void FAssetTools::OnNewCreateRecord(UClass* AssetType, bool bDuplicated)
+{
+	// Don't attempt to report usage stats if analytics isn't available
+	if(AssetType != NULL && FEngineAnalytics::IsAvailable())
+	{
+		TArray<FAnalyticsEventAttribute> Attribs;
+		Attribs.Add(FAnalyticsEventAttribute(TEXT("AssetType"), AssetType->GetName()));
+		Attribs.Add(FAnalyticsEventAttribute(TEXT("Duplicated"), bDuplicated? TEXT("Yes") : TEXT("No")));
+
+		FEngineAnalytics::GetProvider().RecordEvent(TEXT("Editor.Usage.CreateAsset"), Attribs);
+	}
+}
+
 FAssetTools& FAssetTools::Get()
 {
 	FAssetToolsModule& Module = FModuleManager::GetModuleChecked<FAssetToolsModule>("AssetTools");
diff --git a/Engine/Source/Developer/AssetTools/Private/AssetTools.h b/Engine/Source/Developer/AssetTools/Private/AssetTools.h
--- a/Engine/Source/Developer/AssetTools/Private/AssetTools.h
+++ b/Engine/Source/Developer/AssetTools/Private/AssetTools.h
@@ -58,6 +58,9 @@ private:
 	/** Records the time taken for an import and reports it to engine analytics, if available */
 	static void OnNewImportRecord(UClass* AssetType, const FString& FileExtension, bool bSucceeded, bool bWasCancelled, const FDateTime& StartTime);
 
+	/** Records what assets users are creating */
+	static void OnNewCreateRecord(UClass* AssetType, bool bDuplicated);
+
 private:
 	/** The manager to handle renaming assets */
 	TSharedRef<FAssetRenameManager> AssetRenameManager;
diff --git a/Engine/Source/Editor/Kismet/Private/BlueprintEditor.cpp b/Engine/Source/Editor/Kismet/Private/BlueprintEditor.cpp
--- a/Engine/Source/Editor/Kismet/Private/BlueprintEditor.cpp
+++ b/Engine/Source/Editor/Kismet/Private/BlueprintEditor.cpp
@@ -318,6 +318,9 @@ void FBlueprintEditor::AnalyticsTrackNodeEvent( UBlueprint* Blueprint, UEdGraphN
 {
 	if( Blueprint && GraphNode && FEngineAnalytics::IsAvailable() )
 	{
+		// we'd like to see if this was happening in normal blueprint editor or persona 
+		const FString EditorName = Cast<UAnimBlueprint>(Blueprint) != NULL ? TEXT("Persona") : TEXT("BlueprintEditor");
+
 		// Build Node Details
 		const UGeneralProjectSettings& ProjectSettings = *GetDefault<UGeneralProjectSettings>();
 		FString ProjectID = ProjectSettings.ProjectID.ToString();
@@ -343,9 +346,9 @@ void FBlueprintEditor::AnalyticsTrackNodeEvent( UBlueprint* Blueprint, UEdGraphN
 			{
 				NodeAttributes.Add( FAnalyticsEventAttribute( Iter->Key, Iter->Value ));
 			}
-			// Send Analytics event
-			FString EventType = bNodeDelete ?	TEXT( "Editor.Usage.BlueprintEditor.NodeDeleted" ) :
-												TEXT( "Editor.Usage.BlueprintEditor.NodeCreated" );
+			// Send Analytics event 
+			FString EventType = bNodeDelete ?	FString::Printf(TEXT( "Editor.Usage.%s.NodeDeleted" ), *EditorName) :
+												FString::Printf(TEXT( "Editor.Usage.%s.NodeCreated" ), *EditorName);
 			FEngineAnalytics::GetProvider().RecordEvent( EventType, NodeAttributes );
 		}
 	}
diff --git a/Engine/Source/Editor/PhAT/Private/PhAT.cpp b/Engine/Source/Editor/PhAT/Private/PhAT.cpp
--- a/Engine/Source/Editor/PhAT/Private/PhAT.cpp
+++ b/Engine/Source/Editor/PhAT/Private/PhAT.cpp
@@ -13,8 +13,6 @@
 #include "PhATEdSkeletalMeshComponent.h"
 #include "PhAT.h"
 
-
-
 #include "Editor/WorkspaceMenuStructure/Public/WorkspaceMenuStructureModule.h"
 #include "Editor/PropertyEditor/Public/PropertyEditorModule.h"
 #include "Editor/PropertyEditor/Public/IDetailsView.h"
@@ -25,6 +23,9 @@
 #include "Developer/MeshUtilities/Public/MeshUtilities.h"
 #include "UnrealEd.h"
 
+#include "EngineAnalytics.h"
+#include "Runtime/Analytics/Analytics/Public/Interfaces/IAnalyticsProvider.h"
+
 DEFINE_LOG_CATEGORY(LogPhAT);
 #define LOCTEXT_NAMESPACE "PhAT"
 
@@ -2186,6 +2187,8 @@ void FPhAT::OnEditingMode(int32 Mode)
 	// Rebuild the toolbar, as the icons shown will have changed
 	ExtendToolbar();
 	RegenerateMenusAndToolbars();
+
+	OnAddPhatRecord(TEXT("ModeSelected"), false, true);
 }
 
 bool FPhAT::IsEditingMode(int32 Mode) const
@@ -2309,6 +2312,9 @@ void FPhAT::ImpToggleSimulation()
 	{
 		PreviewViewport->GetViewportClient()->SetRealtime(true);
 	}
+
+	// add to analytics record
+	OnAddPhatRecord(TEXT("ToggleSimulate"), true, true);
 }
 
 void FPhAT::OnSetSimulationMode(EPhATSimulationMode Mode)
@@ -3101,4 +3107,30 @@ void FPhAT::OnSelectAll()
 	}
 }
 
+// record if simulating or not, or mode changed or not, or what mode it is in while simulating and what kind of simulation options
+void FPhAT::OnAddPhatRecord(const FString & Action, bool bRecordSimulate, bool bRecordMode )
+{
+	// Don't attempt to report usage stats if analytics isn't available
+	if( Action.IsEmpty() == false && SharedData.IsValid() && FEngineAnalytics::IsAvailable())
+	{
+		TArray<FAnalyticsEventAttribute> Attribs;
+		if (bRecordSimulate)
+		{
+			Attribs.Add(FAnalyticsEventAttribute(TEXT("Simulation"), SharedData->bRunningSimulation? TEXT("ON") : TEXT("OFF")));
+			if ( SharedData->bRunningSimulation )
+			{
+				Attribs.Add(FAnalyticsEventAttribute(TEXT("Selected"), IsSelectedSimulation()? TEXT("ON") : TEXT("OFF")));
+				Attribs.Add(FAnalyticsEventAttribute(TEXT("Gravity"), IsSimulationMode(EPSM_Gravity)? TEXT("ON") : TEXT("OFF")));
+			}
+		}
+
+		if (bRecordMode)
+		{
+			Attribs.Add(FAnalyticsEventAttribute(TEXT("EditMode"), SharedData->EditingMode == FPhATSharedData::PEM_ConstraintEdit? TEXT("Constraint") : TEXT("Body")));
+		}
+
+		FString EventString = FString::Printf(TEXT("Editor.PHAT.%s"), *Action);
+		FEngineAnalytics::GetProvider().RecordEvent(EventString, Attribs);
+	}
+}
 #undef LOCTEXT_NAMESPACE
diff --git a/Engine/Source/Editor/PhAT/Private/PhAT.h b/Engine/Source/Editor/PhAT/Private/PhAT.h
--- a/Engine/Source/Editor/PhAT/Private/PhAT.h
+++ b/Engine/Source/Editor/PhAT/Private/PhAT.h
@@ -290,4 +290,7 @@ private:
 	TArray<EPhysicsType> PhysicsTypeState;
 	void FixPhysicsState();
 	void ImpToggleSimulation();
+
+	/** Records PhAT related data - simulating or mode change */
+	void OnAddPhatRecord(const FString & Action, bool bRecordSimulate, bool bRecordMode );
 };
