Added Usage on the analytic event param.

diff --git a/Engine/Source/Editor/PhAT/Private/PhAT.cpp b/Engine/Source/Editor/PhAT/Private/PhAT.cpp
--- a/Engine/Source/Editor/PhAT/Private/PhAT.cpp
+++ b/Engine/Source/Editor/PhAT/Private/PhAT.cpp
@@ -3129,7 +3129,7 @@ void FPhAT::OnAddPhatRecord(const FString & Action, bool bRecordSimulate, bool b
 			Attribs.Add(FAnalyticsEventAttribute(TEXT("EditMode"), SharedData->EditingMode == FPhATSharedData::PEM_ConstraintEdit? TEXT("Constraint") : TEXT("Body")));
 		}
 
-		FString EventString = FString::Printf(TEXT("Editor.PHAT.%s"), *Action);
+		FString EventString = FString::Printf(TEXT("Editor.Usage.PHAT.%s"), *Action);
 		FEngineAnalytics::GetProvider().RecordEvent(EventString, Attribs);
 	}
 }
