PhysicsAsset dirtied just by simulating/selecting bodies (NUXFeedbackNUXF-225)

diff --git a/Engine/Source/Editor/PhAT/Private/PhAT.cpp b/Engine/Source/Editor/PhAT/Private/PhAT.cpp
--- a/Engine/Source/Editor/PhAT/Private/PhAT.cpp
+++ b/Engine/Source/Editor/PhAT/Private/PhAT.cpp
@@ -1233,15 +1233,15 @@ void FPhAT::BindCommands()
 
 	ToolkitCommands->MapAction(
 		Commands.KinematicAllBodiesBelow,
-		FExecuteAction::CreateSP(this, &FPhAT::SetBodiesBelowSelectedPhysicsType, EPhysicsType::PhysType_Kinematic) );
+		FExecuteAction::CreateSP(this, &FPhAT::SetBodiesBelowSelectedPhysicsType, EPhysicsType::PhysType_Kinematic, true) );
 
 	ToolkitCommands->MapAction(
 		Commands.SimulatedAllBodiesBelow,
-		FExecuteAction::CreateSP(this, &FPhAT::SetBodiesBelowSelectedPhysicsType, EPhysicsType::PhysType_Simulated) );
+		FExecuteAction::CreateSP(this, &FPhAT::SetBodiesBelowSelectedPhysicsType, EPhysicsType::PhysType_Simulated, true) );
 
 	ToolkitCommands->MapAction(
 		Commands.MakeAllBodiesBelowDefault,
-		FExecuteAction::CreateSP(this, &FPhAT::SetBodiesBelowSelectedPhysicsType, EPhysicsType::PhysType_Default) );
+		FExecuteAction::CreateSP(this, &FPhAT::SetBodiesBelowSelectedPhysicsType, EPhysicsType::PhysType_Default, true) );
 
 	ToolkitCommands->MapAction(
 		Commands.DeleteBody,
@@ -1951,7 +1951,7 @@ void FPhAT::AddNewPrimitive(EKCollisionPrimitiveType InPrimitiveType, bool bCopy
 	RefreshPreviewViewport();
 }
 
-void FPhAT::SetBodiesBelowSelectedPhysicsType( EPhysicsType InPhysicsType)
+void FPhAT::SetBodiesBelowSelectedPhysicsType( EPhysicsType InPhysicsType, bool bMarkAsDirty)
 {
 	TArray<int32> Indices;
 	for(int32 i=0; i<SharedData->SelectedBodies.Num(); ++i)
@@ -1959,10 +1959,10 @@ void FPhAT::SetBodiesBelowSelectedPhysicsType( EPhysicsType InPhysicsType)
 		Indices.Add(SharedData->SelectedBodies[i].Index);
 	}
 
-	SetBodiesBelowPhysicsType(InPhysicsType, Indices);
+	SetBodiesBelowPhysicsType(InPhysicsType, Indices, bMarkAsDirty);
 }
 
-void FPhAT::SetBodiesBelowPhysicsType( EPhysicsType InPhysicsType, const TArray<int32> & Indices)
+void FPhAT::SetBodiesBelowPhysicsType( EPhysicsType InPhysicsType, const TArray<int32> & Indices, bool bMarkAsDirty)
 {
 	TArray<int32> BelowBodies;
 	
@@ -1977,7 +1977,10 @@ void FPhAT::SetBodiesBelowPhysicsType( EPhysicsType InPhysicsType, const TArray<
 	{
 		int32 BodyIndex = BelowBodies[i];
 		UBodySetup* BodySetup = SharedData->PhysicsAsset->BodySetup[BodyIndex];
-		BodySetup->Modify();
+		if (bMarkAsDirty)
+		{
+			BodySetup->Modify();
+		}
 
 		BodySetup->PhysicsType = InPhysicsType;
 	}
@@ -2298,7 +2301,6 @@ void FPhAT::FixPhysicsState()
 	{
 		for(int32 i=0; i<PhysicsTypeState.Num(); ++i)
 		{
-			BodySetup[i]->Modify();
 			BodySetup[i]->PhysicsType = PhysicsTypeState[i];
 		}
 	}
@@ -2330,6 +2332,9 @@ bool FPhAT::IsSimulationMode(EPhATSimulationMode Mode) const
 
 void FPhAT::OnToggleSimulation()
 {
+	// this stores current physics types before simulate
+	// and recovers to the previous physics types
+	// so after this one, we can modify physics types fine
 	FixPhysicsState();
 	if (IsSelectedSimulation())
 	{
@@ -2362,7 +2367,6 @@ void FPhAT::OnSelectedSimulation()
 		//first we fix all the bodies
 		for(int32 i=0; i<SharedData->PhysicsAsset->BodySetup.Num(); ++i)
 		{
-			BodySetup[i]->Modify();
 			BodySetup[i]->PhysicsType = PhysType_Kinematic;
 		}
 
@@ -2370,7 +2374,7 @@ void FPhAT::OnSelectedSimulation()
 		if(SharedData->EditingMode == FPhATSharedData::PEM_BodyEdit)
 		{
 			//Bodies already have a function that does this
-			SetBodiesBelowSelectedPhysicsType(PhysType_Simulated);
+			SetBodiesBelowSelectedPhysicsType(PhysType_Simulated, false);
 		}else
 		{
 			//constraints need some more work
@@ -2390,7 +2394,7 @@ void FPhAT::OnSelectedSimulation()
 				}
 			}
 
-			SetBodiesBelowPhysicsType(PhysType_Simulated, BodyIndices);
+			SetBodiesBelowPhysicsType(PhysType_Simulated, BodyIndices, false);
 		}
 	}
 
diff --git a/Engine/Source/Editor/PhAT/Private/PhAT.h b/Engine/Source/Editor/PhAT/Private/PhAT.h
--- a/Engine/Source/Editor/PhAT/Private/PhAT.h
+++ b/Engine/Source/Editor/PhAT/Private/PhAT.h
@@ -139,8 +139,8 @@ private:
 
 	/** Collision editing helper methods */
 	void AddNewPrimitive(EKCollisionPrimitiveType PrimitiveType, bool bCopySelected = false);
-	void SetBodiesBelowSelectedPhysicsType( EPhysicsType InPhysicsType);
-	void SetBodiesBelowPhysicsType( EPhysicsType InPhysicsType, const TArray<int32> & Indices);
+	void SetBodiesBelowSelectedPhysicsType( EPhysicsType InPhysicsType, bool bMarkAsDirty);
+	void SetBodiesBelowPhysicsType( EPhysicsType InPhysicsType, const TArray<int32> & Indices, bool bMarkAsDirty);
 
 	/** Toolbar/menu command methods */
 	bool IsNotSimulation() const;
