Fix incorrect rendering of collisions in phat when bone is scaled

diff --git a/Engine/Source/Editor/PhAT/Private/PhATRender.cpp b/Engine/Source/Editor/PhAT/Private/PhATRender.cpp
--- a/Engine/Source/Editor/PhAT/Private/PhATRender.cpp
+++ b/Engine/Source/Editor/PhAT/Private/PhATRender.cpp
@@ -79,6 +79,7 @@ void UPhATEdSkeletalMeshComponent::RenderAssetTools(const FSceneView* View, clas
 		if (BoneIndex != INDEX_NONE)
 		{
 			FTransform BoneTM = GetBoneTransform(BoneIndex);
+			float Scale = BoneTM.GetScale3D().GetAbsMax();
 			BoneTM.RemoveScaling();
 
 			FKAggregateGeom* AggGeom = &PhysicsAsset->BodySetup[i]->AggGeom;
@@ -90,13 +91,13 @@ void UPhATEdSkeletalMeshComponent::RenderAssetTools(const FSceneView* View, clas
 					PDI->SetHitProxy(new HPhATEdBoneProxy(i, KPT_Sphere, j));
 				}
 
-				FTransform ElemTM = GetPrimitiveTransform(BoneTM, i, KPT_Sphere, j, 1.f);
+				FTransform ElemTM = GetPrimitiveTransform(BoneTM, i, KPT_Sphere, j, Scale);
 
 				//solids are drawn if it's the ViewMode and we're not doing a hit, or if it's hitAndBodyMode
 				if( (CollisionViewMode == FPhATSharedData::PRM_Solid && !bHitTest) || bHitTestAndBodyMode)
 				{
 					UMaterialInterface*	PrimMaterial = GetPrimitiveMaterial(i, KPT_Sphere, j, bHitTestAndBodyMode);
-					AggGeom->SphereElems[j].DrawElemSolid(PDI, ElemTM, 1.f, PrimMaterial->GetRenderProxy(0));
+					AggGeom->SphereElems[j].DrawElemSolid(PDI, ElemTM, Scale, PrimMaterial->GetRenderProxy(0));
 				}
 
 				//wires are never used during hit
@@ -104,7 +105,7 @@ void UPhATEdSkeletalMeshComponent::RenderAssetTools(const FSceneView* View, clas
 				{
 					if (CollisionViewMode == FPhATSharedData::PRM_Solid || CollisionViewMode == FPhATSharedData::PRM_Wireframe)
 					{
-						AggGeom->SphereElems[j].DrawElemWire(PDI, ElemTM, 1.f, GetPrimitiveColor(i, KPT_Sphere, j));
+						AggGeom->SphereElems[j].DrawElemWire(PDI, ElemTM, Scale, GetPrimitiveColor(i, KPT_Sphere, j));
 					}
 				}
 
@@ -122,19 +123,19 @@ void UPhATEdSkeletalMeshComponent::RenderAssetTools(const FSceneView* View, clas
 					PDI->SetHitProxy(new HPhATEdBoneProxy(i, KPT_Box, j));
 				}
 
-				FTransform ElemTM = GetPrimitiveTransform(BoneTM, i, KPT_Box, j, 1.f);
+				FTransform ElemTM = GetPrimitiveTransform(BoneTM, i, KPT_Box, j, Scale);
 
 				if ( (CollisionViewMode == FPhATSharedData::PRM_Solid && !bHitTest) || bHitTestAndBodyMode)
 				{
 					UMaterialInterface*	PrimMaterial = GetPrimitiveMaterial(i, KPT_Box, j, bHitTestAndBodyMode);
-					AggGeom->BoxElems[j].DrawElemSolid(PDI, ElemTM, 1.f, PrimMaterial->GetRenderProxy(0));
+					AggGeom->BoxElems[j].DrawElemSolid(PDI, ElemTM, Scale, PrimMaterial->GetRenderProxy(0));
 				}
 
 				if(!bHitTest)
 				{
 					if (CollisionViewMode == FPhATSharedData::PRM_Solid || CollisionViewMode == FPhATSharedData::PRM_Wireframe)
 					{
-						AggGeom->BoxElems[j].DrawElemWire(PDI, ElemTM, 1.f, GetPrimitiveColor(i, KPT_Box, j));
+						AggGeom->BoxElems[j].DrawElemWire(PDI, ElemTM, Scale, GetPrimitiveColor(i, KPT_Box, j));
 					}
 				}
 				
@@ -152,19 +153,19 @@ void UPhATEdSkeletalMeshComponent::RenderAssetTools(const FSceneView* View, clas
 					PDI->SetHitProxy(new HPhATEdBoneProxy(i, KPT_Sphyl, j));
 				}
 
-				FTransform ElemTM = GetPrimitiveTransform(BoneTM, i, KPT_Sphyl, j, 1.f);
+				FTransform ElemTM = GetPrimitiveTransform(BoneTM, i, KPT_Sphyl, j, Scale);
 
 				if ( (CollisionViewMode == FPhATSharedData::PRM_Solid && !bHitTest) || bHitTestAndBodyMode)
 				{
 					UMaterialInterface*	PrimMaterial = GetPrimitiveMaterial(i, KPT_Sphyl, j, bHitTestAndBodyMode);
-					AggGeom->SphylElems[j].DrawElemSolid(PDI, ElemTM, 1.f, PrimMaterial->GetRenderProxy(0));
+					AggGeom->SphylElems[j].DrawElemSolid(PDI, ElemTM, Scale, PrimMaterial->GetRenderProxy(0));
 				}
 
 				if(!bHitTest)
 				{
 					if (CollisionViewMode == FPhATSharedData::PRM_Solid || CollisionViewMode == FPhATSharedData::PRM_Wireframe)
 					{
-						AggGeom->SphylElems[j].DrawElemWire(PDI, ElemTM, 1.f, GetPrimitiveColor(i, KPT_Sphyl, j));
+						AggGeom->SphylElems[j].DrawElemWire(PDI, ElemTM, Scale, GetPrimitiveColor(i, KPT_Sphyl, j));
 					}
 				}
 
@@ -181,7 +182,7 @@ void UPhATEdSkeletalMeshComponent::RenderAssetTools(const FSceneView* View, clas
 					PDI->SetHitProxy(new HPhATEdBoneProxy(i, KPT_Convex, j));
 				}
 
-				FTransform ElemTM = GetPrimitiveTransform(BoneTM, i, KPT_Convex, j, 1.f);
+				FTransform ElemTM = GetPrimitiveTransform(BoneTM, i, KPT_Convex, j, Scale);
 
 				//convex doesn't have solid draw so render lines if we're in hitTestAndBodyMode
 				if(!bHitTest || bHitTestAndBodyMode)
