phat multi select works with shift

diff --git a/Engine/Source/Editor/PhAT/Private/PhATPreviewViewportClient.cpp b/Engine/Source/Editor/PhAT/Private/PhATPreviewViewportClient.cpp
--- a/Engine/Source/Editor/PhAT/Private/PhATPreviewViewportClient.cpp
+++ b/Engine/Source/Editor/PhAT/Private/PhATPreviewViewportClient.cpp
@@ -368,13 +368,13 @@ void FPhATEdPreviewViewportClient::ProcessClick(class FSceneView& View, class HH
 		{
 			HPhATEdBoneProxy* BoneProxy = (HPhATEdBoneProxy*)HitProxy;
 
-			SharedData->HitBone(BoneProxy->BodyIndex, BoneProxy->PrimType, BoneProxy->PrimIndex, IsCtrlPressed());
+			SharedData->HitBone(BoneProxy->BodyIndex, BoneProxy->PrimType, BoneProxy->PrimIndex, IsCtrlPressed() || IsShiftPressed());
 		}
 		else if (HitProxy && HitProxy->IsA(HPhATEdConstraintProxy::StaticGetType()))
 		{
 			HPhATEdConstraintProxy* ConstraintProxy = (HPhATEdConstraintProxy*)HitProxy;
 
-			SharedData->HitConstraint(ConstraintProxy->ConstraintIndex, IsCtrlPressed());
+			SharedData->HitConstraint(ConstraintProxy->ConstraintIndex, IsCtrlPressed() || IsShiftPressed());
 		}
 		else
 		{	
