Expose orthogonal view in phat

diff --git a/Engine/Source/Editor/PhAT/Private/PhAT.cpp b/Engine/Source/Editor/PhAT/Private/PhAT.cpp
--- a/Engine/Source/Editor/PhAT/Private/PhAT.cpp
+++ b/Engine/Source/Editor/PhAT/Private/PhAT.cpp
@@ -1213,6 +1213,30 @@ void FPhAT::BindCommands()
 		FCanExecuteAction(),
 		FIsActionChecked::CreateSP(this, &FPhAT::IsShowSkeleton));
 
+	ToolkitCommands->MapAction(
+		Commands.PerspectiveView,
+		FExecuteAction::CreateSP(this, &FPhAT::OnViewType, ELevelViewportType::LVT_Perspective),
+		FCanExecuteAction()
+		);
+
+	ToolkitCommands->MapAction(
+		Commands.TopView,
+		FExecuteAction::CreateSP(this, &FPhAT::OnViewType, ELevelViewportType::LVT_OrthoXY),
+		FCanExecuteAction()
+		);
+
+	ToolkitCommands->MapAction(
+		Commands.SideView,
+		FExecuteAction::CreateSP(this, &FPhAT::OnViewType, ELevelViewportType::LVT_OrthoYZ),
+		FCanExecuteAction()
+		);
+
+	ToolkitCommands->MapAction(
+		Commands.FrontView,
+		FExecuteAction::CreateSP(this, &FPhAT::OnViewType, ELevelViewportType::LVT_OrthoXZ),
+		FCanExecuteAction()
+		);
+
 	ToolkitCommands->MapAction(
 		Commands.MakeBodyKinematic,
 		FExecuteAction::CreateSP(this, &FPhAT::OnSetBodyPhysicsType, EPhysicsType::PhysType_Kinematic ),
@@ -2697,6 +2721,14 @@ bool FPhAT::IsPlayAnimation() const
 	return SharedData->EditorSkelComp->IsPlaying();
 }
 
+void FPhAT::OnViewType(ELevelViewportType ViewType)
+{
+	if (PreviewViewport.IsValid())
+	{
+		PreviewViewport->SetViewportType(ViewType);
+	}
+}
+
 void FPhAT::OnShowSkeleton()
 {
 	SharedData->bShowAnimSkel = !SharedData->bShowAnimSkel;
diff --git a/Engine/Source/Editor/PhAT/Private/PhAT.h b/Engine/Source/Editor/PhAT/Private/PhAT.h
--- a/Engine/Source/Editor/PhAT/Private/PhAT.h
+++ b/Engine/Source/Editor/PhAT/Private/PhAT.h
@@ -213,6 +213,7 @@ private:
 	void OnPlayAnimation();
 	bool IsPlayAnimation() const;
 	void OnShowSkeleton();
+	void OnViewType(ELevelViewportType ViewType);
 	bool IsShowSkeleton() const;
 	void OnSetBodyPhysicsType( EPhysicsType InPhysicsType );
 	bool IsBodyPhysicsType( EPhysicsType InPhysicsType );
diff --git a/Engine/Source/Editor/PhAT/Private/PhATActions.cpp b/Engine/Source/Editor/PhAT/Private/PhATActions.cpp
--- a/Engine/Source/Editor/PhAT/Private/PhATActions.cpp
+++ b/Engine/Source/Editor/PhAT/Private/PhATActions.cpp
@@ -71,6 +71,11 @@ void FPhATCommands::RegisterCommands()
 	UI_COMMAND(SelectAllObjects, "Select All Objects", "Select All Objects", EUserInterfaceActionType::Button, FInputGesture(EModifierKey::Control, EKeys::A));
 	UI_COMMAND(HierarchyFilterAll, "All Bones", "Show Entire Hierarchy", EUserInterfaceActionType::Button, FInputGesture());
 	UI_COMMAND(HierarchyFilterBodies, "Bones With Bodies", "Filter Bones With Bodies", EUserInterfaceActionType::Button, FInputGesture());
+	UI_COMMAND(PerspectiveView, "Perspective", "Perspective", EUserInterfaceActionType::Button, FInputGesture());
+	UI_COMMAND(SideView, "Side", "Orthographic view from side", EUserInterfaceActionType::Button, FInputGesture());
+	UI_COMMAND(TopView, "Top", "Orthographic view from top", EUserInterfaceActionType::Button, FInputGesture());
+	UI_COMMAND(FrontView, "Front", "Orthographic view from front", EUserInterfaceActionType::Button, FInputGesture());
+
 
 	UI_COMMAND(SelectionLock, "Lock Selection", "", EUserInterfaceActionType::Button, FInputGesture(EKeys::X));
 
diff --git a/Engine/Source/Editor/PhAT/Private/PhATActions.h b/Engine/Source/Editor/PhAT/Private/PhATActions.h
--- a/Engine/Source/Editor/PhAT/Private/PhATActions.h
+++ b/Engine/Source/Editor/PhAT/Private/PhATActions.h
@@ -83,6 +83,10 @@ public:
 	TSharedPtr<FUICommandInfo> SelectAllObjects;
 	TSharedPtr<FUICommandInfo> HierarchyFilterAll;
 	TSharedPtr<FUICommandInfo> HierarchyFilterBodies;
+	TSharedPtr<FUICommandInfo> PerspectiveView;
+	TSharedPtr<FUICommandInfo> TopView;
+	TSharedPtr<FUICommandInfo> SideView;
+	TSharedPtr<FUICommandInfo> FrontView;
 
 	/** Hotkey only commands */
 	TSharedPtr<FUICommandInfo> SelectionLock;
diff --git a/Engine/Source/Editor/PhAT/Private/PhATPreviewViewportClient.cpp b/Engine/Source/Editor/PhAT/Private/PhATPreviewViewportClient.cpp
--- a/Engine/Source/Editor/PhAT/Private/PhATPreviewViewportClient.cpp
+++ b/Engine/Source/Editor/PhAT/Private/PhATPreviewViewportClient.cpp
@@ -74,7 +74,7 @@ FPhATEdPreviewViewportClient::FPhATEdPreviewViewportClient(TWeakPtr<FPhAT> InPhA
 	
 	SetViewLocationForOrbiting(FVector::ZeroVector);
 
-	SetViewMode(VMI_Lit);
+	SetViewModes(VMI_Lit, VMI_Lit);
 
 	bUsingOrbitCamera = true;
 
@@ -101,7 +101,7 @@ void FPhATEdPreviewViewportClient::DrawCanvas( FViewport& InViewport, FSceneView
 	float W, H;
 	PhATFont->GetCharSize(TEXT('L'), W, H);
 
-	const float XOffset = 150.0f;
+	const float XOffset = 200.0f;
 
 	FCanvasTextItem TextItem( FVector2D::ZeroVector, FText::GetEmpty(), PhATFont, FLinearColor::White );
 
@@ -804,7 +804,7 @@ void FPhATEdPreviewViewportClient::SimMousePress(FViewport* Viewport, bool bCons
 #endif
 	SharedData->LastClickPos = Click.GetClickPos();
 	FHitResult Result(1.f);
-	bool bHit = SharedData->EditorSkelComp->LineTraceComponent(Result, Click.GetOrigin(), Click.GetOrigin() + Click.GetDirection() * SimGrabCheckDistance, FCollisionQueryParams(true));
+	bool bHit = SharedData->EditorSkelComp->LineTraceComponent(Result, Click.GetOrigin() - Click.GetDirection() * SimGrabCheckDistance, Click.GetOrigin() + Click.GetDirection() * SimGrabCheckDistance, FCollisionQueryParams(true));
 
 	if (bHit)
 	{
@@ -870,7 +870,7 @@ void FPhATEdPreviewViewportClient::SimMouseMove(float DeltaX, float DeltaY)
 
 	//Now we project new ScreenPos to xy-plane of SimGrabLocation
 	FVector LocalOffset = View->ViewMatrices.ViewMatrix.TransformPosition(SimGrabLocation + SimGrabZ * SimGrabPush);
-	float ZDistance = fabs(LocalOffset.Z);
+	float ZDistance = GetViewportType() == ELevelViewportType::LVT_Perspective ? fabs(LocalOffset.Z) : 1.f;	//in the ortho case we don't need to do any fixup because there is no perspective
 	WorldDelta = ProjectedDelta * ZDistance;
 	
 	//Now we convert back into WorldPos
diff --git a/Engine/Source/Editor/PhAT/Private/SPhATPreviewToolbar.cpp b/Engine/Source/Editor/PhAT/Private/SPhATPreviewToolbar.cpp
--- a/Engine/Source/Editor/PhAT/Private/SPhATPreviewToolbar.cpp
+++ b/Engine/Source/Editor/PhAT/Private/SPhATPreviewToolbar.cpp
@@ -21,6 +21,16 @@ void SPhATPreviewViewportToolBar::Construct(const FArguments& InArgs)
 		.ForegroundColor(FEditorStyle::GetSlateColor("DefaultForeground"))
 		[
 			SNew(SHorizontalBox)
+			+ SHorizontalBox::Slot()
+			.AutoWidth()
+			.Padding(5.0f, 2.0f)
+			[
+				SNew(SEditorViewportToolbarMenu)
+				.ParentToolBar(SharedThis(this))
+				.Cursor(EMouseCursor::Default)
+				.Label(NSLOCTEXT("PhAT", "CameraMenuTitle_Default", "Camera"))
+				.OnGetMenuContent(this, &SPhATPreviewViewportToolBar::GeneratePerspectiveMenu)
+			]
 			+SHorizontalBox::Slot()
 			.AutoWidth()
 			.Padding(5.0f, 2.0f)
@@ -47,6 +57,22 @@ void SPhATPreviewViewportToolBar::Construct(const FArguments& InArgs)
 	SViewportToolBar::Construct(SViewportToolBar::FArguments());
 }
 
+TSharedRef<SWidget> SPhATPreviewViewportToolBar::GeneratePerspectiveMenu() const
+{
+	const FPhATCommands& Actions = FPhATCommands::Get();
+
+	const bool bInShouldCloseWindowAfterMenuSelection = true;
+	FMenuBuilder PerspectiveMenuBuilder(bInShouldCloseWindowAfterMenuSelection, PhATPtr.Pin()->GetToolkitCommands());
+	{
+		PerspectiveMenuBuilder.AddMenuEntry(Actions.PerspectiveView);
+		PerspectiveMenuBuilder.AddMenuEntry(Actions.TopView);
+		PerspectiveMenuBuilder.AddMenuEntry(Actions.SideView);
+		PerspectiveMenuBuilder.AddMenuEntry(Actions.FrontView);
+	}
+
+	return PerspectiveMenuBuilder.MakeWidget();
+}
+
 TSharedRef<SWidget> SPhATPreviewViewportToolBar::GenerateViewMenu() const
 {
 	const FPhATCommands& Actions = FPhATCommands::Get();
diff --git a/Engine/Source/Editor/PhAT/Private/SPhATPreviewToolbar.h b/Engine/Source/Editor/PhAT/Private/SPhATPreviewToolbar.h
--- a/Engine/Source/Editor/PhAT/Private/SPhATPreviewToolbar.h
+++ b/Engine/Source/Editor/PhAT/Private/SPhATPreviewToolbar.h
@@ -17,6 +17,9 @@ public:
 	void Construct(const FArguments& InArgs);
 
 private:
+
+	/** Generates the toolbar perspective menu content */
+	TSharedRef<SWidget> GeneratePerspectiveMenu() const;
 	/** Generates the toolbar view menu content */
 	TSharedRef<SWidget> GenerateViewMenu() const;
 
diff --git a/Engine/Source/Editor/PhAT/Private/SPhATPreviewViewport.cpp b/Engine/Source/Editor/PhAT/Private/SPhATPreviewViewport.cpp
--- a/Engine/Source/Editor/PhAT/Private/SPhATPreviewViewport.cpp
+++ b/Engine/Source/Editor/PhAT/Private/SPhATPreviewViewport.cpp
@@ -19,6 +19,11 @@ SPhATPreviewViewport::~SPhATPreviewViewport()
 	}
 }
 
+void SPhATPreviewViewport::SetViewportType(ELevelViewportType ViewType)
+{
+	ViewportClient->SetViewportType(ViewType);
+}
+
 void SPhATPreviewViewport::Construct(const FArguments& InArgs)
 {
 	PhATPtr = InArgs._PhAT;
diff --git a/Engine/Source/Editor/PhAT/Private/SPhATPreviewViewport.h b/Engine/Source/Editor/PhAT/Private/SPhATPreviewViewport.h
--- a/Engine/Source/Editor/PhAT/Private/SPhATPreviewViewport.h
+++ b/Engine/Source/Editor/PhAT/Private/SPhATPreviewViewport.h
@@ -41,6 +41,8 @@ public:
 	/** The parent tab where this viewport resides */
 	TWeakPtr<SDockTab> ParentTab;
 
+	void SetViewportType(ELevelViewportType ViewType);
+
 private:
 	/** Pointer back to the PhysicsAsset editor tool that owns us */
 	TWeakPtr<FPhAT> PhATPtr;
