phat mirroring phase 1

diff --git a/Engine/Source/Editor/PhAT/Private/PhAT.cpp b/Engine/Source/Editor/PhAT/Private/PhAT.cpp
--- a/Engine/Source/Editor/PhAT/Private/PhAT.cpp
+++ b/Engine/Source/Editor/PhAT/Private/PhAT.cpp
@@ -1343,6 +1343,17 @@ void FPhAT::BindCommands()
 	ToolkitCommands->MapAction(
 		Commands.HierarchyFilterBodies,
 		FExecuteAction::CreateSP(this, &FPhAT::SetHierarchyFilter, PHFM_Bodies));
+
+	ToolkitCommands->MapAction(
+		Commands.Mirror,
+		FExecuteAction::CreateSP(this, &FPhAT::Mirror),
+		FCanExecuteAction::CreateSP(this, &FPhAT::IsNotSimulation)
+		);
+}
+
+void FPhAT::Mirror()
+{
+	SharedData->Mirror();
 }
 
 TSharedRef<ITableRow> FPhAT::OnGenerateRowForTree(FTreeElemPtr Item, const TSharedRef<STableViewBase>& OwnerTable)
diff --git a/Engine/Source/Editor/PhAT/Private/PhAT.h b/Engine/Source/Editor/PhAT/Private/PhAT.h
--- a/Engine/Source/Editor/PhAT/Private/PhAT.h
+++ b/Engine/Source/Editor/PhAT/Private/PhAT.h
@@ -230,6 +230,7 @@ private:
 	void OnToggleSwing2();
 	void OnToggleTwist();
 	void OnFocusSelection();
+	void Mirror();
 
 	//menu commands
 	void OnSelectAll();
diff --git a/Engine/Source/Editor/PhAT/Private/PhATActions.cpp b/Engine/Source/Editor/PhAT/Private/PhATActions.cpp
--- a/Engine/Source/Editor/PhAT/Private/PhATActions.cpp
+++ b/Engine/Source/Editor/PhAT/Private/PhATActions.cpp
@@ -75,6 +75,7 @@ void FPhATCommands::RegisterCommands()
 	UI_COMMAND(SideView, "Side", "Orthographic view from side", EUserInterfaceActionType::Button, FInputGesture());
 	UI_COMMAND(TopView, "Top", "Orthographic view from top", EUserInterfaceActionType::Button, FInputGesture());
 	UI_COMMAND(FrontView, "Front", "Orthographic view from front", EUserInterfaceActionType::Button, FInputGesture());
+	UI_COMMAND(Mirror, "Mirror", "Finds the body on the other side and duplicates constraint and body", EUserInterfaceActionType::Button, FInputGesture(EKeys::M));
 
 
 	UI_COMMAND(SelectionLock, "Lock Selection", "", EUserInterfaceActionType::Button, FInputGesture(EKeys::X));
diff --git a/Engine/Source/Editor/PhAT/Private/PhATActions.h b/Engine/Source/Editor/PhAT/Private/PhATActions.h
--- a/Engine/Source/Editor/PhAT/Private/PhATActions.h
+++ b/Engine/Source/Editor/PhAT/Private/PhATActions.h
@@ -87,6 +87,7 @@ public:
 	TSharedPtr<FUICommandInfo> TopView;
 	TSharedPtr<FUICommandInfo> SideView;
 	TSharedPtr<FUICommandInfo> FrontView;
+	TSharedPtr<FUICommandInfo> Mirror;
 
 	/** Hotkey only commands */
 	TSharedPtr<FUICommandInfo> SelectionLock;
diff --git a/Engine/Source/Editor/PhAT/Private/PhATSharedData.cpp b/Engine/Source/Editor/PhAT/Private/PhATSharedData.cpp
--- a/Engine/Source/Editor/PhAT/Private/PhATSharedData.cpp
+++ b/Engine/Source/Editor/PhAT/Private/PhATSharedData.cpp
@@ -204,6 +204,121 @@ void FPhATSharedData::Initialize()
 	EnableSimulation(false);
 }
 
+void FPhATSharedData::CopyConstraintProperties(UPhysicsConstraintTemplate * FromConstraintSetup, UPhysicsConstraintTemplate * ToConstraintSetup)
+{
+	//We want to copy frame2 relative to frame1. To do this we need to go from body2 into body1, compute relative.
+	//Then apply relative inside body1 of the destination constraint, and finally move it into body2 of destination constraint
+
+	//In total there are 4 bodies, body a and b are the frame1 frame2 bodies of the copied
+	// body c and d are frame 1 and frame 2 of the destination constraint
+	//we use the body letter to denote which space we're in, for example Frame1a would be frame1 inside BodyA
+	//Frame1b would be frame 1 inside body b. In this case BodyA * Frame1a == BodyB * Frame1b
+
+	FTransform Frame1A = FromConstraintSetup->DefaultInstance.GetRefFrame(EConstraintFrame::Frame1);
+	FTransform BodyA = GetConstraintBodyTM(FromConstraintSetup, EConstraintFrame::Frame1);
+	FTransform BodyB = GetConstraintBodyTM(FromConstraintSetup, EConstraintFrame::Frame2);
+	FTransform BodyAFrame1A = GetConstraintWorldTM(FromConstraintSetup, EConstraintFrame::Frame1);
+	FTransform BodyBFrame2B = GetConstraintWorldTM(FromConstraintSetup, EConstraintFrame::Frame2);
+
+	FTransform Frame1C = ToConstraintSetup->DefaultInstance.GetRefFrame(EConstraintFrame::Frame1);
+	FTransform BodyC = GetConstraintBodyTM(ToConstraintSetup, EConstraintFrame::Frame1);
+	FTransform BodyD = GetConstraintBodyTM(ToConstraintSetup, EConstraintFrame::Frame2);
+
+	FTransform FromF1AToF2A = BodyBFrame2B * BodyA.Inverse() * Frame1A.Inverse();
+	FTransform Frame2C = FromF1AToF2A * Frame1C;
+	FTransform Frame2D = Frame2C* BodyC * BodyD.Inverse();
+
+
+	ToConstraintSetup->Modify();
+	FConstraintInstance OldInstance = ToConstraintSetup->DefaultInstance;
+	ToConstraintSetup->DefaultInstance.CopyConstraintParamsFrom(&FromConstraintSetup->DefaultInstance);
+
+	// recover certain data that we'd like to keep - i.e. bone indices
+	// those still should stay
+	ToConstraintSetup->DefaultInstance.ConstraintIndex = OldInstance.ConstraintIndex;
+	ToConstraintSetup->DefaultInstance.ConstraintData = OldInstance.ConstraintData;
+	ToConstraintSetup->DefaultInstance.JointName = OldInstance.JointName;
+	ToConstraintSetup->DefaultInstance.ConstraintBone1 = OldInstance.ConstraintBone1;
+	ToConstraintSetup->DefaultInstance.ConstraintBone2 = OldInstance.ConstraintBone2;
+	ToConstraintSetup->DefaultInstance.Pos1 = OldInstance.Pos1;
+	ToConstraintSetup->DefaultInstance.Pos2 = OldInstance.Pos2;
+
+	//frame1 stays the same
+	ToConstraintSetup->DefaultInstance.PriAxis1 = OldInstance.PriAxis1;
+	ToConstraintSetup->DefaultInstance.SecAxis1 = OldInstance.SecAxis1;
+
+	//frame2 is copied but relative to frame1
+	ToConstraintSetup->DefaultInstance.PriAxis2 = Frame2D.GetUnitAxis(EAxis::X);
+	ToConstraintSetup->DefaultInstance.SecAxis2 = Frame2D.GetUnitAxis(EAxis::Y);
+}
+
+struct FMirrorInfo
+{
+	FName BoneName;
+	int32 BoneIndex;
+	int32 BodyIndex;
+	int32 ConstraintIndex;
+	FMirrorInfo()
+	{
+		BoneIndex = INDEX_NONE;
+		BodyIndex = INDEX_NONE;
+		ConstraintIndex = INDEX_NONE;
+		BoneName = NAME_None;
+	}
+};
+
+void FPhATSharedData::Mirror()
+{
+	TArray<FMirrorInfo> MirrorInfos;
+	if (EditingMode == PEM_BodyEdit)	//grab all selected bodies
+	{
+		for (const FSelection & Selection : SelectedBodies)
+		{
+			MirrorInfos.AddUninitialized();
+			FMirrorInfo & MirrorInfo = MirrorInfos[MirrorInfos.Num() - 1];
+			MirrorInfo.BoneName = PhysicsAsset->BodySetup[Selection.Index]->BoneName;
+			MirrorInfo.BodyIndex = Selection.Index;
+			MirrorInfo.ConstraintIndex = PhysicsAsset->FindConstraintIndex(MirrorInfo.BoneName);
+		}
+	}
+	else if (EditingMode == PEM_ConstraintEdit)	//grab all selected constraints
+	{
+		for (const FSelection & Selection : SelectedConstraints)
+		{
+			MirrorInfos.AddUninitialized();
+			FMirrorInfo & MirrorInfo = MirrorInfos[MirrorInfos.Num() - 1];
+			MirrorInfo.BoneName = PhysicsAsset->ConstraintSetup[Selection.Index]->DefaultInstance.ConstraintBone1;
+			MirrorInfo.BodyIndex = PhysicsAsset->FindBodyIndex(MirrorInfo.BoneName);
+			MirrorInfo.ConstraintIndex = Selection.Index;
+		}
+	}
+
+	for (FMirrorInfo & MirrorInfo : MirrorInfos)	//mirror all selected bodies/constraints
+	{
+		int32 BoneIndex = EditorSkelMesh->RefSkeleton.FindBoneIndex(MirrorInfo.BoneName);
+
+		int32 MirrorBoneIndex = PhysicsAsset->FindMirroredBone(EditorSkelMesh, BoneIndex);
+		if (MirrorBoneIndex != INDEX_NONE)
+		{
+			UBodySetup * SrcBody = PhysicsAsset->BodySetup[MirrorInfo.BodyIndex];
+			const FScopedTransaction Transaction(NSLOCTEXT("PhAT", "MirrorBody", "MirrorBody"));
+			MakeNewBody(MirrorBoneIndex, false);
+
+			int32 MirrorBodyIndex = PhysicsAsset->FindControllingBodyIndex(EditorSkelMesh, MirrorBoneIndex);
+
+			UBodySetup * DestBody = PhysicsAsset->BodySetup[MirrorBodyIndex];
+			DestBody->Modify();
+			DestBody->CopyBodyPropertiesFrom(SrcBody);
+
+			int32 MirrorConstraintIndex = PhysicsAsset->FindConstraintIndex(DestBody->BoneName);
+			UPhysicsConstraintTemplate * FromConstraint = PhysicsAsset->ConstraintSetup[MirrorInfo.ConstraintIndex];
+			UPhysicsConstraintTemplate * ToConstraint = PhysicsAsset->ConstraintSetup[MirrorConstraintIndex];
+			CopyConstraintProperties(FromConstraint, ToConstraint);
+		}
+	}
+	
+}
+
 FPhATSharedData::EPhATRenderMode FPhATSharedData::GetCurrentMeshViewMode()
 {
 	if (bRunningSimulation)
@@ -777,7 +892,7 @@ void FPhATSharedData::InitConstraintSetup(UPhysicsConstraintTemplate* Constraint
 	SetCollisionBetween(ChildBodyIndex, ParentBodyIndex, false);
 }
 
-void FPhATSharedData::MakeNewBody(int32 NewBoneIndex)
+void FPhATSharedData::MakeNewBody(int32 NewBoneIndex, bool bAutoSelect)
 {
 	FName NewBoneName = EditorSkelMesh->RefSkeleton.GetBoneName(NewBoneIndex);
 
@@ -886,7 +1001,11 @@ void FPhATSharedData::MakeNewBody(int32 NewBoneIndex)
 	// update the tree
 	HierarchyChangedEvent.Broadcast();
 
-	SetSelectedBodyAnyPrim(NewBodyIndex);
+	if (bAutoSelect)
+	{
+		SetSelectedBodyAnyPrim(NewBodyIndex);
+	}
+	
 
 	RefreshPhysicsAssetChange(PhysicsAsset);
 }
@@ -928,55 +1047,11 @@ void FPhATSharedData::PasteConstraintProperties()
 
 	UPhysicsConstraintTemplate* FromConstraintSetup = CopiedConstraintTemplate;
 
-	//We want to copy frame2 relative to frame1. To do this we need to go from body2 into body1, compute relative.
-	//Then apply relative inside body1 of the destination constraint, and finally move it into body2 of destination constraint
-	
-	//In total there are 4 bodies, body a and b are the frame1 frame2 bodies of the copied
-	// body c and d are frame 1 and frame 2 of the destination constraint
-	//we use the body letter to denote which space we're in, for example Frame1a would be frame1 inside BodyA
-	//Frame1b would be frame 1 inside body b. In this case BodyA * Frame1a == BodyB * Frame1b
-
-	FTransform Frame1A = FromConstraintSetup->DefaultInstance.GetRefFrame(EConstraintFrame::Frame1);
-	FTransform BodyA = GetConstraintBodyTM(FromConstraintSetup, EConstraintFrame::Frame1);
-	FTransform BodyB = GetConstraintBodyTM(FromConstraintSetup, EConstraintFrame::Frame2);
-	FTransform BodyAFrame1A = GetConstraintWorldTM(FromConstraintSetup, EConstraintFrame::Frame1);
-	FTransform BodyBFrame2B = GetConstraintWorldTM(FromConstraintSetup, EConstraintFrame::Frame2);
-
 	for(int32 i=0; i<SelectedConstraints.Num(); ++i)
 	{
 		// If we are showing instance properties - copy instance properties. If showing setup, just copy setup properties.
 		UPhysicsConstraintTemplate* ToConstraintSetup = PhysicsAsset->ConstraintSetup[SelectedConstraints[i].Index];
-
-		FTransform Frame1C = ToConstraintSetup->DefaultInstance.GetRefFrame(EConstraintFrame::Frame1);
-		FTransform BodyC = GetConstraintBodyTM(ToConstraintSetup, EConstraintFrame::Frame1);
-		FTransform BodyD = GetConstraintBodyTM(ToConstraintSetup, EConstraintFrame::Frame2);
-
-		FTransform FromF1AToF2A = BodyBFrame2B * BodyA.Inverse() * Frame1A.Inverse();
-		FTransform Frame2C = FromF1AToF2A * Frame1C;
-		FTransform Frame2D = Frame2C* BodyC * BodyD.Inverse();
-
-
-		ToConstraintSetup->Modify();
-		FConstraintInstance OldInstance = ToConstraintSetup->DefaultInstance;
-		ToConstraintSetup->DefaultInstance.CopyConstraintParamsFrom(&FromConstraintSetup->DefaultInstance);
-
-		// recover certain data that we'd like to keep - i.e. bone indices
-		// those still should stay
-		ToConstraintSetup->DefaultInstance.ConstraintIndex		= OldInstance.ConstraintIndex;
-		ToConstraintSetup->DefaultInstance.ConstraintData		= OldInstance.ConstraintData;
-		ToConstraintSetup->DefaultInstance.JointName			= OldInstance.JointName;
-		ToConstraintSetup->DefaultInstance.ConstraintBone1		= OldInstance.ConstraintBone1;
-		ToConstraintSetup->DefaultInstance.ConstraintBone2		= OldInstance.ConstraintBone2;
-		ToConstraintSetup->DefaultInstance.Pos1					= OldInstance.Pos1;
-		ToConstraintSetup->DefaultInstance.Pos2					= OldInstance.Pos2;
-
-		//frame1 stays the same
-		ToConstraintSetup->DefaultInstance.PriAxis1 = OldInstance.PriAxis1;
-		ToConstraintSetup->DefaultInstance.SecAxis1 = OldInstance.SecAxis1;
-
-		//frame2 is copied but relative to frame1
-		ToConstraintSetup->DefaultInstance.PriAxis2 = Frame2D.GetUnitAxis(EAxis::X);
-		ToConstraintSetup->DefaultInstance.SecAxis2 = Frame2D.GetUnitAxis(EAxis::Y);
+		CopyConstraintProperties(FromConstraintSetup, ToConstraintSetup);
 	}
 	
 }
diff --git a/Engine/Source/Editor/PhAT/Private/PhATSharedData.h b/Engine/Source/Editor/PhAT/Private/PhATSharedData.h
--- a/Engine/Source/Editor/PhAT/Private/PhATSharedData.h
+++ b/Engine/Source/Editor/PhAT/Private/PhATSharedData.h
@@ -106,12 +106,14 @@ public:
 	void DeleteCurrentPrim();
 	void DeleteBody(int32 DelBodyIndex, bool bRefreshComponent=true);
 	void RefreshPhysicsAssetChange(const UPhysicsAsset * InPhysAsset);
-	void MakeNewBody(int32 NewBoneIndex);
+	void MakeNewBody(int32 NewBoneIndex, bool bAutoSelect = true);
 	void CopyBody();
 	void CopyConstraint();
 	void PasteBodyProperties();
 	bool WeldSelectedBodies(bool bWeld = true);
 
+	void Mirror();
+
 	/** Misc */
 	void ToggleInstanceProperties();
 	void ToggleSimulation();
@@ -139,6 +141,8 @@ private:
 	/** Hides the transform widget if the simulation is running, else restores it to its previous mode before simulation */
 	void UpdateTransformWidgetVisibilityForSimulationMode(bool bEnableSimulation);
 
+	void CopyConstraintProperties(UPhysicsConstraintTemplate * FromConstraintSetup, UPhysicsConstraintTemplate * ToConstraintSetup);
+
 public:
 	/** Callback for handling selection changes */
 	DECLARE_EVENT_TwoParams(FPhATSharedData, FSelectionChanged, UObject*, FSelection *);
diff --git a/Engine/Source/Runtime/Engine/Classes/PhysicsEngine/PhysicsAsset.h b/Engine/Source/Runtime/Engine/Classes/PhysicsEngine/PhysicsAsset.h
--- a/Engine/Source/Runtime/Engine/Classes/PhysicsEngine/PhysicsAsset.h
+++ b/Engine/Source/Runtime/Engine/Classes/PhysicsEngine/PhysicsAsset.h
@@ -130,6 +130,7 @@ public:
 	ENGINE_API int32		FindParentBodyIndex(class USkeletalMesh * skelMesh, int32 StartBoneIndex) const;
 	ENGINE_API int32		FindConstraintIndex(FName ConstraintName);
 	FName					FindConstraintBoneName(int32 ConstraintIndex);
+	ENGINE_API int32		FindMirroredBone(class USkeletalMesh* skelMesh, int32 BoneIndex);
 
 	/** Utility for getting indices of all bodies below (and including) the one with the supplied name. */
 	ENGINE_API void			GetBodyIndicesBelow(TArray<int32>& OutBodyIndices, FName InBoneName, USkeletalMesh* InSkelMesh, bool bIncludeParent = true);
diff --git a/Engine/Source/Runtime/Engine/Private/PhysicsEngine/PhysicsAsset.cpp b/Engine/Source/Runtime/Engine/Private/PhysicsEngine/PhysicsAsset.cpp
--- a/Engine/Source/Runtime/Engine/Private/PhysicsEngine/PhysicsAsset.cpp
+++ b/Engine/Source/Runtime/Engine/Private/PhysicsEngine/PhysicsAsset.cpp
@@ -261,6 +261,33 @@ FName UPhysicsAsset::FindConstraintBoneName(int32 ConstraintIndex)
 	return ConstraintSetup[ConstraintIndex]->DefaultInstance.JointName;
 }
 
+int32 UPhysicsAsset::FindMirroredBone(class USkeletalMesh* skelMesh, int32 BoneIndex)
+{
+	if (BoneIndex == INDEX_NONE)
+	{
+		return INDEX_NONE;
+	}
+	//we try to find the mirroed bone using several approaches. The first is to look for the same name but with _R instead of _L or vise versa
+	FName BoneName = skelMesh->RefSkeleton.GetBoneName(BoneIndex);
+	FString BoneNameString = BoneName.ToString();
+
+	bool bIsLeft = BoneNameString.Find("_L", ESearchCase::IgnoreCase, ESearchDir::FromEnd) == (BoneNameString.Len() - 2);	//has _L at the end
+	bool bIsRight = BoneNameString.Find("_R", ESearchCase::IgnoreCase, ESearchDir::FromEnd) == (BoneNameString.Len() - 2);	//has _R at the end
+	bool bMirrorConvention = bIsLeft || bIsRight;
+
+	//if the bone follows our left right naming convention then let's try and find its mirrored bone
+	int32 BoneIndexMirrored = INDEX_NONE;
+	if (bMirrorConvention)
+	{
+		FString BoneNameMirrored = BoneNameString.LeftChop(2);
+		BoneNameMirrored.Append(bIsLeft ? "_R" : "_L");
+
+		BoneIndexMirrored = skelMesh->RefSkeleton.FindBoneIndex(FName(*BoneNameMirrored));
+	}
+
+	return BoneIndexMirrored;
+}
+
 void UPhysicsAsset::GetBodyIndicesBelow(TArray<int32>& OutBodyIndices, FName InBoneName, USkeletalMesh* SkelMesh, bool bIncludeParent /*= true*/)
 {
 	int32 BaseIndex = SkelMesh->RefSkeleton.FindBoneIndex(InBoneName);
