UE-1338 EDITOR: PhAT: CRASH: Copying a static mesh collision and then duplicating it can result in a crash

diff --git a/Engine/Source/Editor/PhAT/Private/PhAT.cpp b/Engine/Source/Editor/PhAT/Private/PhAT.cpp
--- a/Engine/Source/Editor/PhAT/Private/PhAT.cpp
+++ b/Engine/Source/Editor/PhAT/Private/PhAT.cpp
@@ -2005,6 +2005,32 @@ void FPhAT::AddNewPrimitive(EKCollisionPrimitiveType InPrimitiveType, bool bCopy
 					SphylElem->Radius = BodySetup->AggGeom.SphylElems[SharedData->GetSelectedBody()->PrimitiveIndex].Radius;
 				}
 			}
+			else if (PrimitiveType == KPT_Convex)
+			{
+				check(bCopySelected); //only support copying for Convex primitive, as there is no default vertex data
+
+				NewPrimIndex = BodySetup->AggGeom.ConvexElems.Add(FKConvexElem());
+				NewSelection[i].PrimitiveType = KPT_Convex;
+				NewSelection[i].PrimitiveIndex = NewPrimIndex;
+				FKConvexElem* ConvexElem = &BodySetup->AggGeom.ConvexElems[NewPrimIndex];
+
+				ConvexElem->SetTransform(BodySetup->AggGeom.ConvexElems[SharedData->GetSelectedBody()->PrimitiveIndex].GetTransform());
+
+				// Copy all of the vertices of the convex element
+				for (FVector v : BodySetup->AggGeom.ConvexElems[SharedData->GetSelectedBody()->PrimitiveIndex].VertexData)
+				{
+					v.X += PhAT::DuplicateXOffset;
+					ConvexElem->VertexData.Add(v);
+				}
+				ConvexElem->UpdateElemBox();
+
+				BodySetup->InvalidatePhysicsData();
+				BodySetup->CreatePhysicsMeshes();
+			}
+			else
+			{
+				check(0);  //unrecognized primitive type
+			}
 		}
 	} // ScopedTransaction
 
@@ -3078,7 +3104,7 @@ void FPhAT::OnAssetSelectedFromStaticMeshAssetPicker( const FAssetData& AssetDat
 	
 		UStaticMesh* SM = Cast<UStaticMesh>(AssetData.GetAsset());
 
-		if (SM && SM->BodySetup && SM->BodySetup->AggGeom.ConvexElems.Num() > 0)
+		if (SM && SM->BodySetup && SM->BodySetup->AggGeom.GetElementCount() > 0)
 		{
 			BaseSetup->AddCollisionFrom(SM->BodySetup);
 
