Improve missing skeletal mesh warning message

diff --git a/Engine/Source/Editor/PhAT/Private/PhATSharedData.cpp b/Engine/Source/Editor/PhAT/Private/PhATSharedData.cpp
--- a/Engine/Source/Editor/PhAT/Private/PhATSharedData.cpp
+++ b/Engine/Source/Editor/PhAT/Private/PhATSharedData.cpp
@@ -100,7 +100,7 @@ void FPhATSharedData::Initialize()
 
 		FMessageDialog::Open(EAppMsgType::Ok, FText::Format(
 				NSLOCTEXT("UnrealEd", "Error_PhysicsAssetHasNoSkelMesh",
-				"Warning: Physics Asset has no default SkeletalMesh assigned!  For now, a simple default skeletal mesh ({0}) will be used.  You should repair the DefaultSkeletalMesh using UnrealPhAT (Edit -> Change Default SkeletalMesh) before saving this asset."),
+				"Warning: Physics Asset has no skeletal mesh assigned!  For now, a simple default skeletal mesh ({0}) will be used.  You can fix this by opening PhAT, selecting the appropriate skeletal mesh in the content browser, and using (Asset -> Change Mesh) before saving this asset."),
 				FText::FromString(PreviewMesh->GetFullName())));
 	}
 
