Integration fix for code that was moved

diff --git a/Engine/Source/Editor/PhAT/Private/SPhATPreviewViewport.cpp b/Engine/Source/Editor/PhAT/Private/SPhATPreviewViewport.cpp
--- a/Engine/Source/Editor/PhAT/Private/SPhATPreviewViewport.cpp
+++ b/Engine/Source/Editor/PhAT/Private/SPhATPreviewViewport.cpp
@@ -75,6 +75,5 @@ TSharedPtr<SWidget> SPhATPreviewViewport::MakeViewportToolbar()
 {
 	return SNew(SPhATPreviewViewportToolBar)
 			.PhATPtr(PhATPtr)
-			.Visibility(this, &SPhATPreviewViewport::GetWidgetVisibility)
 			.IsEnabled(FSlateApplication::Get().GetNormalExecutionAttribute());
 }
