Turn off bUpdateJointsFromAnimation in PhAT. Fixes orientation motor not working correctly.

diff --git a/Engine/Source/Editor/PhAT/Private/PhATRender.cpp b/Engine/Source/Editor/PhAT/Private/PhATRender.cpp
--- a/Engine/Source/Editor/PhAT/Private/PhATRender.cpp
+++ b/Engine/Source/Editor/PhAT/Private/PhATRender.cpp
@@ -42,7 +42,7 @@ UPhATEdSkeletalMeshComponent::UPhATEdSkeletalMeshComponent(const FObjectInitiali
 
 	// this is because in phat editor, you'd like to see fixed bones to be fixed without animation force update
 	KinematicBonesUpdateType = EKinematicBonesUpdateToPhysics::SkipSimulatingBones;
-	bUpdateJointsFromAnimation = true;
+	bUpdateJointsFromAnimation = false;
 	ForcedLodModel = 1;
 
 	static FName CollisionProfileName(TEXT("PhysicsActor"));
