Remove manual transform toolbar. Replace with the functionality found in base class. UE-9577

diff --git a/Engine/Source/Editor/PhAT/Private/PhAT.cpp b/Engine/Source/Editor/PhAT/Private/PhAT.cpp
--- a/Engine/Source/Editor/PhAT/Private/PhAT.cpp
+++ b/Engine/Source/Editor/PhAT/Private/PhAT.cpp
@@ -195,6 +195,7 @@ void FPhAT::InitPhAT(const EToolkitMode::Type Mode, const TSharedPtr< class IToo
 	HierarchyFilterMode = PHFM_All;
 	SelectedAnimation = NULL;
 	SelectedSimulation = false;
+	BeforeSimulationWidgetMode = FWidget::EWidgetMode::WM_None;
 
 	SharedData = MakeShareable(new FPhATSharedData);
 
@@ -575,29 +576,6 @@ void FPhAT::AddReferencedObjects(FReferenceCollector& Collector)
 	Collector.AddReferencedObject(SharedData->MouseHandle);
 }
 
-void FPhAT::CycleTransformMode()
-{
-	if ( !SharedData->bRunningSimulation )
-	{
-		switch( SharedData->WidgetMode )
-		{
-		case FWidget::WM_Translate:
-			SharedData->WidgetMode = FWidget::WM_Rotate;
-			break;
-		case FWidget::WM_Rotate:
-			SharedData->WidgetMode = FWidget::WM_Scale;
-			break;
-		case FWidget::WM_Scale:
-			SharedData->WidgetMode = FWidget::WM_Translate;
-			break;	
-		default:
-			SharedData->WidgetMode = FWidget::WM_Translate;
-		}
-
-		RefreshPreviewViewport();
-	}
-}
-
 void FPhAT::PostUndo(bool bSuccess)
 {
 	SharedData->PostUndo();
@@ -808,14 +786,7 @@ void FPhAT::ExtendToolbar()
 			}
 			ToolbarBuilder.EndSection();
 	
-			ToolbarBuilder.BeginSection("PhATTransform");
-			{
-				ToolbarBuilder.AddToolBarButton(Commands.PhATTranslationMode);
-				ToolbarBuilder.AddToolBarButton(Commands.PhATRotationMode);
-				ToolbarBuilder.AddToolBarButton(Commands.PhATScaleMode);
-			}
-			ToolbarBuilder.EndSection();
-
+	
 			if ( InPhATEditingMode == FPhATSharedData::PEM_BodyEdit )
 			{
 				ToolbarBuilder.BeginSection("PhATCollision");
@@ -990,36 +961,6 @@ void FPhAT::BindCommands()
 		FCanExecuteAction::CreateSP(this, &FPhAT::IsNotSimulation),
 		FIsActionChecked::CreateSP(this, &FPhAT::IsEditingMode, (int32)FPhATSharedData::PEM_ConstraintEdit));
 
-	ToolkitCommands->MapAction(
-		Commands.MovementSpace_Local,
-		FExecuteAction::CreateSP(this, &FPhAT::OnMovementSpace, (int32)COORD_Local),
-		FCanExecuteAction::CreateSP(this, &FPhAT::IsNotSimulation),
-		FIsActionChecked::CreateSP(this, &FPhAT::IsMovementSpace, (int32)COORD_Local));
-
-	ToolkitCommands->MapAction(
-		Commands.MovementSpace_World,
-		FExecuteAction::CreateSP(this, &FPhAT::OnMovementSpace, (int32)COORD_World),
-		FCanExecuteAction::CreateSP(this, &FPhAT::IsNotSimulation),
-		FIsActionChecked::CreateSP(this, &FPhAT::IsMovementSpace, (int32)COORD_World));
-
-	ToolkitCommands->MapAction(
-		Commands.PhATTranslationMode,
-		FExecuteAction::CreateSP(this, &FPhAT::OnTransformMode, (int32)FWidget::WM_Translate),
-		FCanExecuteAction::CreateSP(this, &FPhAT::IsNotSimulation),
-		FIsActionChecked::CreateSP(this, &FPhAT::IsTransformMode, (int32)FWidget::WM_Translate));
-
-	ToolkitCommands->MapAction(
-		Commands.PhATRotationMode,
-		FExecuteAction::CreateSP(this, &FPhAT::OnTransformMode, (int32)FWidget::WM_Rotate),
-		FCanExecuteAction::CreateSP(this, &FPhAT::IsNotSimulation),
-		FIsActionChecked::CreateSP(this, &FPhAT::IsTransformMode, (int32)FWidget::WM_Rotate));
-
-	ToolkitCommands->MapAction(
-		Commands.PhATScaleMode,
-		FExecuteAction::CreateSP(this, &FPhAT::OnTransformMode, (int32)FWidget::WM_Scale),
-		FCanExecuteAction::CreateSP(this, &FPhAT::IsNotSimulation),
-		FIsActionChecked::CreateSP(this, &FPhAT::IsTransformMode, (int32)FWidget::WM_Scale));
-
 	ToolkitCommands->MapAction(
 		Commands.CopyProperties,
 		FExecuteAction::CreateSP(this, &FPhAT::OnCopyProperties),
@@ -1351,11 +1292,6 @@ void FPhAT::BindCommands()
 		FExecuteAction::CreateSP(this, &FPhAT::OnFocusSelection),
 		FCanExecuteAction());
 	
-	ToolkitCommands->MapAction(
-		Commands.CycleTransformMode,
-		FExecuteAction::CreateSP(this, &FPhAT::CycleTransformMode),
-		FCanExecuteAction::CreateSP(this, &FPhAT::IsNotSimulation));
-
 	ToolkitCommands->MapAction(
 		Commands.SelectAllObjects,
 		FExecuteAction::CreateSP(this, &FPhAT::OnSelectAll));
@@ -2271,12 +2207,6 @@ void FPhAT::OnEditingMode(int32 Mode)
 		SharedData->EditingMode = FPhATSharedData::PEM_ConstraintEdit;
 		RefreshHierachyTree();
 		SharedData->SetSelectedConstraint(INDEX_NONE, true);
-
-		// Scale isn't valid for constraints!
-		if (SharedData->WidgetMode == FWidget::WM_Scale)
-		{
-			SharedData->WidgetMode = FWidget::WM_Translate;
-		}
 	}
 	
 	RefreshPreviewViewport();
@@ -2293,30 +2223,6 @@ bool FPhAT::IsEditingMode(int32 Mode) const
 	return (FPhATSharedData::EPhATEditingMode)Mode == SharedData->EditingMode;
 }
 
-void FPhAT::OnMovementSpace(int32 Mode)
-{
-	SharedData->MovementSpace = (ECoordSystem)Mode;
-
-	RefreshPreviewViewport();
-}
-
-bool FPhAT::IsMovementSpace(int32 Mode) const
-{
-	return (ECoordSystem)Mode == SharedData->MovementSpace;
-}
-
-void FPhAT::OnTransformMode(int32 Mode)
-{
-	SharedData->WidgetMode = (FWidget::EWidgetMode)Mode;
-
-	RefreshPreviewViewport();
-}
-
-bool FPhAT::IsTransformMode(int32 Mode) const
-{
-	return SharedData->WidgetMode == (FWidget::EWidgetMode)Mode;
-}
-
 void FPhAT::OnCopyProperties()
 {
 	if(SharedData->EditingMode == FPhATSharedData::PEM_BodyEdit)
@@ -2402,6 +2308,17 @@ void FPhAT::FixPhysicsState()
 
 void FPhAT::ImpToggleSimulation()
 {
+	TSharedPtr<FPhATEdPreviewViewportClient> PreviewClient = PreviewViewport->GetViewportClient();
+	if(!SharedData->bRunningSimulation)
+	{
+		BeforeSimulationWidgetMode = PreviewClient->GetWidgetMode();
+		PreviewClient->SetWidgetMode(FWidget::EWidgetMode::WM_None);
+	}else
+	{
+		PreviewClient->SetWidgetMode(BeforeSimulationWidgetMode);
+	}
+	
+
 	SharedData->ToggleSimulation();
 
 	if (!PreviewViewport->GetViewportClient()->IsRealtime() && !IsPIERunning())
diff --git a/Engine/Source/Editor/PhAT/Private/PhAT.h b/Engine/Source/Editor/PhAT/Private/PhAT.h
--- a/Engine/Source/Editor/PhAT/Private/PhAT.h
+++ b/Engine/Source/Editor/PhAT/Private/PhAT.h
@@ -84,9 +84,6 @@ public:
 
 	/** FGCObject interface */
 	virtual void AddReferencedObjects(FReferenceCollector& Collector) override;
-	
-	/** Cycles through transform modes. */
-	void CycleTransformMode();
 
 	/** Returns whether a PIE session is running. */
 	static bool IsPIERunning();
@@ -163,10 +160,6 @@ private:
 	void OnApplyPhysicalMaterial();
 	void OnEditingMode(int32 Mode);
 	bool IsEditingMode(int32 Mode) const;
-	void OnMovementSpace(int32 Mode);
-	bool IsMovementSpace(int32 Mode) const;
-	void OnTransformMode(int32 Mode);
-	bool IsTransformMode(int32 Mode) const;
 	void OnCopyProperties();
 	bool IsCopyProperties() const;
 	bool CanCopyProperties() const;
@@ -311,4 +304,5 @@ private:
 	FText GetRecordStatusTooltip() const;
 	FText GetRecordStatusLabel() const;
 	FText GetRecordMenuLabel() const;
+	FWidget::EWidgetMode BeforeSimulationWidgetMode;
 };
diff --git a/Engine/Source/Editor/PhAT/Private/PhATActions.cpp b/Engine/Source/Editor/PhAT/Private/PhATActions.cpp
--- a/Engine/Source/Editor/PhAT/Private/PhATActions.cpp
+++ b/Engine/Source/Editor/PhAT/Private/PhATActions.cpp
@@ -13,11 +13,6 @@ void FPhATCommands::RegisterCommands()
 	UI_COMMAND(ApplyPhysicalMaterial, "Apply PhysMat", "Applies Currently Selected Physical Material In Content Browser To All Bodies", EUserInterfaceActionType::Button, FInputGesture());
 	UI_COMMAND(EditingMode_Body, "Body Mode", "Body Editing Mode", EUserInterfaceActionType::RadioButton, FInputGesture());
 	UI_COMMAND(EditingMode_Constraint, "Constraint Mode", "Constraint Editing Mode", EUserInterfaceActionType::RadioButton, FInputGesture());
-	UI_COMMAND(MovementSpace_Local, "Local", "Local Movement Space", EUserInterfaceActionType::RadioButton, FInputGesture());
-	UI_COMMAND(MovementSpace_World, "World", "World Movement Space", EUserInterfaceActionType::RadioButton, FInputGesture());
-	UI_COMMAND(PhATTranslationMode, "Translation", "Translation Mode", EUserInterfaceActionType::ToggleButton, FInputGesture(EKeys::W));
-	UI_COMMAND(PhATRotationMode, "Rotation", "Rotation Mode", EUserInterfaceActionType::ToggleButton, FInputGesture(EKeys::E));
-	UI_COMMAND(PhATScaleMode, "Scale", "Scale Mode", EUserInterfaceActionType::ToggleButton, FInputGesture(EKeys::R));
 	UI_COMMAND(CopyProperties, "Copy Properties", "Copy Properties: Copy Properties Of Currently Selected Object To Next Selected Object", EUserInterfaceActionType::Button, FInputGesture(EModifierKey::Control, EKeys::C));
 	UI_COMMAND(PasteProperties, "Paste Properties", "Paste Properties: Copy Properties Of Currently Selected Object To Next Selected Object", EUserInterfaceActionType::Button, FInputGesture(EModifierKey::Control, EKeys::V));
 	UI_COMMAND(InstanceProperties, "Instance Properties", "Instance Properties: Displays Instance Properties When In Body Editing Mode", EUserInterfaceActionType::ToggleButton, FInputGesture(EKeys::I));
@@ -91,5 +86,4 @@ void FPhATCommands::RegisterCommands()
 	UI_COMMAND(ToggleSwing2, "Toggle Swing2 Constraint", "", EUserInterfaceActionType::Button, FInputGesture(EKeys::Three));
 	UI_COMMAND(ToggleTwist, "Toggle Twist Constraint", "", EUserInterfaceActionType::Button, FInputGesture(EKeys::Two));
 	UI_COMMAND(FocusOnSelection, "Focus the viewport on the current selection", "", EUserInterfaceActionType::Button, FInputGesture(EKeys::F));
-	UI_COMMAND(CycleTransformMode, "Cycle the transform mode", "", EUserInterfaceActionType::Button, FInputGesture(EKeys::SpaceBar));
 }
diff --git a/Engine/Source/Editor/PhAT/Private/PhATActions.h b/Engine/Source/Editor/PhAT/Private/PhATActions.h
--- a/Engine/Source/Editor/PhAT/Private/PhATActions.h
+++ b/Engine/Source/Editor/PhAT/Private/PhATActions.h
@@ -24,11 +24,6 @@ public:
 	TSharedPtr<FUICommandInfo> ApplyPhysicalMaterial;
 	TSharedPtr<FUICommandInfo> EditingMode_Body;
 	TSharedPtr<FUICommandInfo> EditingMode_Constraint;
-	TSharedPtr<FUICommandInfo> MovementSpace_Local;
-	TSharedPtr<FUICommandInfo> MovementSpace_World;
-	TSharedPtr<FUICommandInfo> PhATTranslationMode;
-	TSharedPtr<FUICommandInfo> PhATRotationMode;
-	TSharedPtr<FUICommandInfo> PhATScaleMode;
 	TSharedPtr<FUICommandInfo> Snap;
 	TSharedPtr<FUICommandInfo> CopyProperties;
 	TSharedPtr<FUICommandInfo> PasteProperties;
@@ -99,7 +94,6 @@ public:
 	TSharedPtr<FUICommandInfo> ToggleSwing2;
 	TSharedPtr<FUICommandInfo> ToggleTwist;
 	TSharedPtr<FUICommandInfo> FocusOnSelection;
-	TSharedPtr<FUICommandInfo> CycleTransformMode;
 	
 
 	/** Initialize commands */
diff --git a/Engine/Source/Editor/PhAT/Private/PhATPreviewViewportClient.cpp b/Engine/Source/Editor/PhAT/Private/PhATPreviewViewportClient.cpp
--- a/Engine/Source/Editor/PhAT/Private/PhATPreviewViewportClient.cpp
+++ b/Engine/Source/Editor/PhAT/Private/PhATPreviewViewportClient.cpp
@@ -34,6 +34,9 @@ FPhATEdPreviewViewportClient::FPhATEdPreviewViewportClient(TWeakPtr<FPhAT> InPhA
 {
 	check(PhATPtr.IsValid());
 
+	ModeTools->SetWidgetMode(FWidget::EWidgetMode::WM_Translate);
+	ModeTools->SetCoordSystem(COORD_Local);
+
 	bAllowedToMoveCamera = true;
 
 	// Setup defaults for the common draw helper.
@@ -137,10 +140,10 @@ void FPhATEdPreviewViewportClient::DrawCanvas( FViewport& InViewport, FSceneView
 		TextItem.Text = NSLOCTEXT("UnrealEd", "Lock", "LOCK");
 	}else if(SharedData->EditingMode == FPhATSharedData::PEM_ConstraintEdit)
 	{
-		if(SharedData->WidgetMode == FWidget::WM_Translate)
+		if(GetWidgetMode() == FWidget::WM_Translate)
 		{
 			TextItem.Text = NSLOCTEXT("UnrealEd", "SingleMove", "hold ALT to move a single reference frame");
-		}else if(SharedData->WidgetMode == FWidget::WM_Rotate)
+		}else if(GetWidgetMode() == FWidget::WM_Rotate)
 		{
 			TextItem.Text = NSLOCTEXT("UnrealEd", "DoubleRotate", "hold ALT to rotate both reference frames");
 		}
@@ -477,13 +480,13 @@ bool FPhATEdPreviewViewportClient::InputWidgetDelta( FViewport* Viewport, EAxisL
 				SelectedObject.WidgetTM = SharedData->GetConstraintMatrix(SelectedObject.Index, EConstraintFrame::Frame2, 1.f);
 			}
 
-			if ( SharedData->WidgetMode == FWidget::WM_Translate )
+			if ( GetWidgetMode() == FWidget::WM_Translate )
 			{
 				FVector Dir = SelectedObject.WidgetTM.InverseTransformVector( Drag.GetSafeNormal() );
 				FVector DragVec = Dir * Drag.Size();
 				SelectedObject.ManipulateTM.AddToTranslation( DragVec );
 			}
-			else if ( SharedData->WidgetMode == FWidget::WM_Rotate )
+			else if ( GetWidgetMode() == FWidget::WM_Rotate )
 			{
 				FVector Axis; 
 				float Angle;
@@ -497,7 +500,7 @@ bool FPhATEdPreviewViewportClient::InputWidgetDelta( FViewport* Viewport, EAxisL
 
 				SelectedObject.ManipulateTM = FTransform( Result );
 			}
-			else if ( SharedData->WidgetMode == FWidget::WM_Scale && SharedData->EditingMode == FPhATSharedData::PEM_BodyEdit) // Scaling only valid for bodies.
+			else if ( GetWidgetMode() == FWidget::WM_Scale && SharedData->EditingMode == FPhATSharedData::PEM_BodyEdit) // Scaling only valid for bodies.
 			{
 				ModifyPrimitiveSize(SelectedObject.Index, SelectedObject.PrimitiveType, SelectedObject.PrimitiveIndex, Scale );
 			}
@@ -509,7 +512,7 @@ bool FPhATEdPreviewViewportClient::InputWidgetDelta( FViewport* Viewport, EAxisL
 				ConstraintSetup->DefaultInstance.SetRefFrame(EConstraintFrame::Frame2, SelectedObject.ManipulateTM * StartManParentConTM);
 
 				//Rotation by default only rotates one frame, but translation by default moves both
-				bool bMultiFrame = (IsAltPressed() && SharedData->WidgetMode == FWidget::WM_Rotate) || (!IsAltPressed() && SharedData->WidgetMode == FWidget::WM_Translate);
+				bool bMultiFrame = (IsAltPressed() && GetWidgetMode() == FWidget::WM_Rotate) || (!IsAltPressed() && GetWidgetMode() == FWidget::WM_Translate);
 				
 				if (bMultiFrame)
 				{
@@ -560,21 +563,22 @@ void FPhATEdPreviewViewportClient::TrackingStopped()
 
 FWidget::EWidgetMode FPhATEdPreviewViewportClient::GetWidgetMode() const
 {
-	FWidget::EWidgetMode WidgetMode = FWidget::WM_None;
+	FWidget::EWidgetMode ReturnWidgetMode = FWidget::WM_None;
+	FWidget::EWidgetMode WidgetMode = FEditorViewportClient::GetWidgetMode();
 
 	if( SharedData->EditingMode == FPhATSharedData::PEM_BodyEdit && SharedData->GetSelectedBody()) 
 	{
-		WidgetMode = SharedData->WidgetMode;
+		ReturnWidgetMode = WidgetMode;
 	}
 	else if( 
 		SharedData->EditingMode == FPhATSharedData::PEM_ConstraintEdit 
 		&& SharedData->GetSelectedConstraint()
-		&& SharedData->WidgetMode != FWidget::WM_Scale )
+		&& WidgetMode != FWidget::WM_Scale )
 	{
-		WidgetMode = SharedData->WidgetMode;
+		ReturnWidgetMode = WidgetMode;
 	}
 
-	return WidgetMode;
+	return ReturnWidgetMode;
 }
 
 FVector FPhATEdPreviewViewportClient::GetWidgetLocation() const
@@ -643,7 +647,7 @@ FMatrix FPhATEdPreviewViewportClient::GetWidgetCoordSystem() const
 
 ECoordSystem FPhATEdPreviewViewportClient::GetWidgetCoordSystemSpace() const
 {
-	return SharedData->WidgetMode == FWidget::WM_Scale ? COORD_Local : SharedData->MovementSpace;
+	return GetWidgetMode() == FWidget::WM_Scale ? COORD_Local : ModeTools->GetCoordSystem();;
 }
 
 void FPhATEdPreviewViewportClient::Tick(float DeltaSeconds)
diff --git a/Engine/Source/Editor/PhAT/Private/PhATSharedData.cpp b/Engine/Source/Editor/PhAT/Private/PhATSharedData.cpp
--- a/Engine/Source/Editor/PhAT/Private/PhATSharedData.cpp
+++ b/Engine/Source/Editor/PhAT/Private/PhATSharedData.cpp
@@ -23,7 +23,6 @@ FPhATSharedData::FPhATSharedData()
 	, CopiedBodySetup(NULL)
 	, CopiedConstraintTemplate(NULL)
 	, bInsideSelChange(false)
-	, WidgetModeBeforeSimulation(FWidget::WM_None)
 {
 	// Editor variables
 	BodyEdit_MeshViewMode = PRM_Solid;
@@ -38,7 +37,6 @@ FPhATSharedData::FPhATSharedData()
 	Sim_CollisionViewMode = PRM_Wireframe;
 	Sim_ConstraintViewMode = PCV_None;
 
-	MovementSpace = COORD_Local;
 	EditingMode = PEM_BodyEdit;
 
 	bShowCOM = false;
@@ -1450,19 +1448,6 @@ void FPhATSharedData::ToggleSimulation()
 	bRunningSimulation = !bRunningSimulation;
 }
 
-void FPhATSharedData::UpdateTransformWidgetVisibilityForSimulationMode(bool bEnableSimulation)
-{
-	if ( bEnableSimulation )
-	{
-		WidgetModeBeforeSimulation = WidgetMode;
-		WidgetMode = FWidget::WM_None;
-	}
-	else
-	{
-		WidgetMode = WidgetModeBeforeSimulation;
-	}
-}
-
 void FPhATSharedData::EnableSimulation(bool bEnableSimulation)
 {
 	if (bEnableSimulation)
@@ -1508,10 +1493,6 @@ void FPhATSharedData::EnableSimulation(bool bEnableSimulation)
 			SetSelectedConstraint(INDEX_NONE, true);
 		}
 	}
-	if( bEnableSimulation != bRunningSimulation )
-	{
-		UpdateTransformWidgetVisibilityForSimulationMode(bEnableSimulation);
-	}
 }
 
 void FPhATSharedData::OpenNewBodyDlg()
diff --git a/Engine/Source/Editor/PhAT/Private/PhATSharedData.h b/Engine/Source/Editor/PhAT/Private/PhATSharedData.h
--- a/Engine/Source/Editor/PhAT/Private/PhATSharedData.h
+++ b/Engine/Source/Editor/PhAT/Private/PhATSharedData.h
@@ -139,9 +139,6 @@ private:
 
 	void EnableSimulation(bool bEnableSimulation);
 
-	/** Hides the transform widget if the simulation is running, else restores it to its previous mode before simulation */
-	void UpdateTransformWidgetVisibilityForSimulationMode(bool bEnableSimulation);
-
 	void CopyConstraintProperties(UPhysicsConstraintTemplate * FromConstraintSetup, UPhysicsConstraintTemplate * ToConstraintSetup);
 
 public:
@@ -243,8 +240,6 @@ public:
 
 	/** Various mode flags */
 	EPhATEditingMode EditingMode;
-	FWidget::EWidgetMode WidgetMode;
-	ECoordSystem MovementSpace;
 	
 	/** We have a different set of view setting per editing mode */
 	EPhATRenderMode BodyEdit_MeshViewMode;
@@ -270,6 +265,4 @@ public:
 	FVector LastClickDirection;
 #endif
 	FIntPoint LastClickPos;
-private:
-	FWidget::EWidgetMode WidgetModeBeforeSimulation;
 };
diff --git a/Engine/Source/Editor/PhAT/Private/SPhATPreviewToolbar.cpp b/Engine/Source/Editor/PhAT/Private/SPhATPreviewToolbar.cpp
--- a/Engine/Source/Editor/PhAT/Private/SPhATPreviewToolbar.cpp
+++ b/Engine/Source/Editor/PhAT/Private/SPhATPreviewToolbar.cpp
@@ -3,12 +3,14 @@
 #include "PhATModule.h"
 #include "PhAT.h"
 #include "Editor/UnrealEd/Public/SViewportToolBar.h"
+#include "Editor/UnrealEd/Public/STransformViewportToolbar.h"
 #include "Editor/UnrealEd/Public/SEditorViewportToolBarMenu.h"
+#include "Editor/UnrealEd/Public/SEditorViewportViewMenu.h"
 #include "PhATActions.h"
 #include "SPhATPreviewToolbar.h"
 
 
-void SPhATPreviewViewportToolBar::Construct(const FArguments& InArgs)
+void SPhATPreviewViewportToolBar::Construct(const FArguments& InArgs, TSharedPtr<class SEditorViewport> InRealViewport)
 {
 	PhATPtr = InArgs._PhATPtr;
 	
@@ -53,6 +55,15 @@ void SPhATPreviewViewportToolBar::Construct(const FArguments& InArgs)
 				.Label(NSLOCTEXT("PhAT", "ModesMenuTitle_Default", "Modes"))
 				.OnGetMenuContent(this, &SPhATPreviewViewportToolBar::GenerateModesMenu) 
 			]
+
+			+ SHorizontalBox::Slot()
+			.Padding(3.0f, 1.0f)
+			.HAlign(HAlign_Right)
+			[
+				SNew(STransformViewportToolBar)
+				.Viewport(InRealViewport)
+				.CommandList(InRealViewport->GetCommandList())
+			]
 		]
 	];
 
@@ -101,30 +112,6 @@ TSharedRef<SWidget> SPhATPreviewViewportToolBar::GenerateModesMenu() const
 	const bool bInShouldCloseWindowAfterMenuSelection = true;
 	FMenuBuilder ModesMenuBuilder(bInShouldCloseWindowAfterMenuSelection, PhATPtr.Pin()->GetToolkitCommands());
 	{
-		// Editing and movement modes
-		{
-			struct Local
-			{
-				static void BuildMovementModeMenu(FMenuBuilder& Menu)
-				{
-					const FPhATCommands& Commands = FPhATCommands::Get();
-
-					Menu.BeginSection("PhATMovementMode", NSLOCTEXT("PhAT", "MovementModeHeader", "Movement Mode"));
-					{
-						Menu.AddMenuEntry(Commands.MovementSpace_Local);
-						Menu.AddMenuEntry(Commands.MovementSpace_World);
-					}
-					Menu.EndSection();
-				}
-			};
-
-			ModesMenuBuilder.BeginSection("PhATSubMenus");
-			{
-				ModesMenuBuilder.AddSubMenu(NSLOCTEXT("PhAT", "MovementModeSubMenu", "Movement Mode"), FText::GetEmpty(), FNewMenuDelegate::CreateStatic(&Local::BuildMovementModeMenu));
-			}
-			ModesMenuBuilder.EndSection();
-		}
-
 		// Mesh, collision and constraint rendering modes
 		{
 			struct Local
diff --git a/Engine/Source/Editor/PhAT/Private/SPhATPreviewToolbar.h b/Engine/Source/Editor/PhAT/Private/SPhATPreviewToolbar.h
--- a/Engine/Source/Editor/PhAT/Private/SPhATPreviewToolbar.h
+++ b/Engine/Source/Editor/PhAT/Private/SPhATPreviewToolbar.h
@@ -14,7 +14,7 @@ public:
 		SLATE_ARGUMENT(TWeakPtr<FPhAT>, PhATPtr)
 	SLATE_END_ARGS()
 
-	void Construct(const FArguments& InArgs);
+	void Construct(const FArguments& InArgs, TSharedPtr<class SEditorViewport> InRealViewport);
 
 private:
 
diff --git a/Engine/Source/Editor/PhAT/Private/SPhATPreviewViewport.cpp b/Engine/Source/Editor/PhAT/Private/SPhATPreviewViewport.cpp
--- a/Engine/Source/Editor/PhAT/Private/SPhATPreviewViewport.cpp
+++ b/Engine/Source/Editor/PhAT/Private/SPhATPreviewViewport.cpp
@@ -73,7 +73,7 @@ TSharedRef<FEditorViewportClient> SPhATPreviewViewport::MakeEditorViewportClient
 
 TSharedPtr<SWidget> SPhATPreviewViewport::MakeViewportToolbar()
 {
-	return SNew(SPhATPreviewViewportToolBar)
+	return SNew(SPhATPreviewViewportToolBar, SharedThis(this))
 			.PhATPtr(PhATPtr)
 			.IsEnabled(FSlateApplication::Get().GetNormalExecutionAttribute());
 }
