Fix non-unity build

diff --git a/Engine/Source/Editor/PhAT/Private/SPhATPreviewToolbar.cpp b/Engine/Source/Editor/PhAT/Private/SPhATPreviewToolbar.cpp
--- a/Engine/Source/Editor/PhAT/Private/SPhATPreviewToolbar.cpp
+++ b/Engine/Source/Editor/PhAT/Private/SPhATPreviewToolbar.cpp
@@ -8,6 +8,7 @@
 #include "Editor/UnrealEd/Public/SEditorViewportViewMenu.h"
 #include "PhATActions.h"
 #include "SPhATPreviewToolbar.h"
+#include "Editor/UnrealEd//Public/SEditorViewport.h"
 
 
 void SPhATPreviewViewportToolBar::Construct(const FArguments& InArgs, TSharedPtr<class SEditorViewport> InRealViewport)
