Fix PhAT widget movement not matching body movement for scaled bones.

diff --git a/Engine/Source/Editor/PhAT/Private/PhATPreviewViewportClient.cpp b/Engine/Source/Editor/PhAT/Private/PhATPreviewViewportClient.cpp
--- a/Engine/Source/Editor/PhAT/Private/PhATPreviewViewportClient.cpp
+++ b/Engine/Source/Editor/PhAT/Private/PhATPreviewViewportClient.cpp
@@ -466,14 +466,16 @@ bool FPhATEdPreviewViewportClient::InputWidgetDelta( FViewport* Viewport, EAxisL
 		FPhATSharedData::FSelection & SelectedObject = SelectedObjects[i];
 		if( SharedData->bManipulating )
 		{
+			float BoneScale = 1.f;
 			if (SharedData->EditingMode == FPhATSharedData::PEM_BodyEdit) /// BODY EDITING ///
 			{
 				int32 BoneIndex = SharedData->EditorSkelComp->GetBoneIndex(SharedData->PhysicsAsset->BodySetup[SelectedObject.Index]->BoneName);
 
 				FTransform BoneTM = SharedData->EditorSkelComp->GetBoneTransform(BoneIndex);
+				BoneScale = BoneTM.GetScale3D().GetAbsMax();
 				BoneTM.RemoveScaling();
 
-				SelectedObject.WidgetTM = SharedData->EditorSkelComp->GetPrimitiveTransform(BoneTM, SelectedObject.Index, SelectedObject.PrimitiveType, SelectedObject.PrimitiveIndex, 1.f );
+				SelectedObject.WidgetTM = SharedData->EditorSkelComp->GetPrimitiveTransform(BoneTM, SelectedObject.Index, SelectedObject.PrimitiveType, SelectedObject.PrimitiveIndex, BoneScale);
 			}
 			else  /// CONSTRAINT EDITING ///
 			{
@@ -483,7 +485,7 @@ bool FPhATEdPreviewViewportClient::InputWidgetDelta( FViewport* Viewport, EAxisL
 			if ( GetWidgetMode() == FWidget::WM_Translate )
 			{
 				FVector Dir = SelectedObject.WidgetTM.InverseTransformVector( Drag.GetSafeNormal() );
-				FVector DragVec = Dir * Drag.Size();
+				FVector DragVec = Dir * Drag.Size() / BoneScale;
 				SelectedObject.ManipulateTM.AddToTranslation( DragVec );
 			}
 			else if ( GetWidgetMode() == FWidget::WM_Rotate )
@@ -594,9 +596,10 @@ FVector FPhATEdPreviewViewportClient::GetWidgetLocation() const
 		int32 BoneIndex = SharedData->EditorSkelComp->GetBoneIndex(SharedData->PhysicsAsset->BodySetup[SharedData->GetSelectedBody()->Index]->BoneName);
 
 		FTransform BoneTM = SharedData->EditorSkelComp->GetBoneTransform(BoneIndex);
+		const float Scale = BoneTM.GetScale3D().GetAbsMax();
 		BoneTM.RemoveScaling();
 
-		return SharedData->EditorSkelComp->GetPrimitiveTransform(BoneTM, SharedData->GetSelectedBody()->Index, SharedData->GetSelectedBody()->PrimitiveType, SharedData->GetSelectedBody()->PrimitiveIndex, 1.f).GetTranslation();
+		return SharedData->EditorSkelComp->GetPrimitiveTransform(BoneTM, SharedData->GetSelectedBody()->Index, SharedData->GetSelectedBody()->PrimitiveType, SharedData->GetSelectedBody()->PrimitiveIndex, Scale).GetTranslation();
 	}
 	else  /// CONSTRAINT EDITING ///
 	{
