Unify Slate's FInputGesture and Engine's FInputChord as FInputChord defined in Slate

diff --git a/Engine/Source/Editor/PListEditor/Public/SPlistEditor.h b/Engine/Source/Editor/PListEditor/Public/SPlistEditor.h
--- a/Engine/Source/Editor/PListEditor/Public/SPlistEditor.h
+++ b/Engine/Source/Editor/PListEditor/Public/SPlistEditor.h
@@ -90,17 +90,17 @@ public:
 	/** Initialize Commands */
 	virtual void RegisterCommands() override
 	{
-		UI_COMMAND( NewCommand, "New", "Creates a new plist file", EUserInterfaceActionType::Button, FInputGesture(EModifierKey::Control, EKeys::N) );
-		UI_COMMAND( OpenCommand, "Open", "Opens an existing plist file", EUserInterfaceActionType::Button, FInputGesture(EModifierKey::Control, EKeys::O) );
-		UI_COMMAND( SaveCommand, "Save", "Saves the current plist file", EUserInterfaceActionType::Button, FInputGesture(EModifierKey::Control, EKeys::S) );
-		UI_COMMAND( SaveAsCommand, "Save As", "Saves the current plist file to a specific location", EUserInterfaceActionType::Button, FInputGesture(EModifierKey::Control | EModifierKey::Shift, EKeys::S) );
-		UI_COMMAND( DeleteSelectedCommand, "Remove Selected", "Removed the selected entries from the plist", EUserInterfaceActionType::Button, FInputGesture(EModifierKey::Control, EKeys::R) );
-		UI_COMMAND( MoveUpCommand, "Move Up", "Moves the selected entry up within its parent", EUserInterfaceActionType::Button, FInputGesture(EModifierKey::Control, EKeys::U) );
-		UI_COMMAND( MoveDownCommand, "Move Down", "Moves the selected entry down within its parent", EUserInterfaceActionType::Button, FInputGesture(EModifierKey::Control, EKeys::D) );
-		UI_COMMAND( AddDictionaryCommand, "Add Dictionary", "Adds a new dictionary to the selected file or array", EUserInterfaceActionType::Button, FInputGesture());
-		UI_COMMAND( AddArrayCommand, "Add Array", "Adds a new array to the selected file, array, or dictionary", EUserInterfaceActionType::Button, FInputGesture());
-		UI_COMMAND( AddStringCommand, "Add String", "Adds a new string to the selected file, array, or dictionary", EUserInterfaceActionType::Button, FInputGesture());
-		UI_COMMAND( AddBooleanCommand, "Add Boolean", "Adds a new boolean to the selected file, array, or dictionary", EUserInterfaceActionType::Button, FInputGesture());
+		UI_COMMAND( NewCommand, "New", "Creates a new plist file", EUserInterfaceActionType::Button, FInputChord(EModifierKey::Control, EKeys::N) );
+		UI_COMMAND( OpenCommand, "Open", "Opens an existing plist file", EUserInterfaceActionType::Button, FInputChord(EModifierKey::Control, EKeys::O) );
+		UI_COMMAND( SaveCommand, "Save", "Saves the current plist file", EUserInterfaceActionType::Button, FInputChord(EModifierKey::Control, EKeys::S) );
+		UI_COMMAND( SaveAsCommand, "Save As", "Saves the current plist file to a specific location", EUserInterfaceActionType::Button, FInputChord(EModifierKey::Control | EModifierKey::Shift, EKeys::S) );
+		UI_COMMAND( DeleteSelectedCommand, "Remove Selected", "Removed the selected entries from the plist", EUserInterfaceActionType::Button, FInputChord(EModifierKey::Control, EKeys::R) );
+		UI_COMMAND( MoveUpCommand, "Move Up", "Moves the selected entry up within its parent", EUserInterfaceActionType::Button, FInputChord(EModifierKey::Control, EKeys::U) );
+		UI_COMMAND( MoveDownCommand, "Move Down", "Moves the selected entry down within its parent", EUserInterfaceActionType::Button, FInputChord(EModifierKey::Control, EKeys::D) );
+		UI_COMMAND( AddDictionaryCommand, "Add Dictionary", "Adds a new dictionary to the selected file or array", EUserInterfaceActionType::Button, FInputChord());
+		UI_COMMAND( AddArrayCommand, "Add Array", "Adds a new array to the selected file, array, or dictionary", EUserInterfaceActionType::Button, FInputChord());
+		UI_COMMAND( AddStringCommand, "Add String", "Adds a new string to the selected file, array, or dictionary", EUserInterfaceActionType::Button, FInputChord());
+		UI_COMMAND( AddBooleanCommand, "Add Boolean", "Adds a new boolean to the selected file, array, or dictionary", EUserInterfaceActionType::Button, FInputChord());
 	}
 };
 
diff --git a/Engine/Source/Editor/Persona/Private/AnimViewportLODCommands.cpp b/Engine/Source/Editor/Persona/Private/AnimViewportLODCommands.cpp
--- a/Engine/Source/Editor/Persona/Private/AnimViewportLODCommands.cpp
+++ b/Engine/Source/Editor/Persona/Private/AnimViewportLODCommands.cpp
@@ -7,11 +7,11 @@
 
 void FAnimViewportLODCommands::RegisterCommands()
 {
-	UI_COMMAND( LODAuto, "LOD Auto", "Automatically select LOD", EUserInterfaceActionType::RadioButton, FInputGesture() );
-	UI_COMMAND( LOD0, "LOD 0", "Force select LOD 0", EUserInterfaceActionType::RadioButton, FInputGesture() );
-	UI_COMMAND( LOD1, "LOD 1", "Force select LOD 1", EUserInterfaceActionType::RadioButton, FInputGesture() );
-	UI_COMMAND( LOD2, "LOD 2", "Force select LOD 2", EUserInterfaceActionType::RadioButton, FInputGesture() );
-	UI_COMMAND( LOD3, "LOD 3", "Force select LOD 3", EUserInterfaceActionType::RadioButton, FInputGesture() );
+	UI_COMMAND( LODAuto, "LOD Auto", "Automatically select LOD", EUserInterfaceActionType::RadioButton, FInputChord() );
+	UI_COMMAND( LOD0, "LOD 0", "Force select LOD 0", EUserInterfaceActionType::RadioButton, FInputChord() );
+	UI_COMMAND( LOD1, "LOD 1", "Force select LOD 1", EUserInterfaceActionType::RadioButton, FInputChord() );
+	UI_COMMAND( LOD2, "LOD 2", "Force select LOD 2", EUserInterfaceActionType::RadioButton, FInputChord() );
+	UI_COMMAND( LOD3, "LOD 3", "Force select LOD 3", EUserInterfaceActionType::RadioButton, FInputChord() );
 }
 
 #undef LOCTEXT_NAMESPACE
\ No newline at end of file
diff --git a/Engine/Source/Editor/Persona/Private/AnimViewportMenuCommands.cpp b/Engine/Source/Editor/Persona/Private/AnimViewportMenuCommands.cpp
--- a/Engine/Source/Editor/Persona/Private/AnimViewportMenuCommands.cpp
+++ b/Engine/Source/Editor/Persona/Private/AnimViewportMenuCommands.cpp
@@ -7,18 +7,18 @@
 
 void FAnimViewportMenuCommands::RegisterCommands()
 {
-	UI_COMMAND( Auto, "Auto", "Preview what is active on the current window", EUserInterfaceActionType::RadioButton, FInputGesture() );
-	UI_COMMAND( Lock, "Lock", "Lock to current preview", EUserInterfaceActionType::RadioButton, FInputGesture() );
+	UI_COMMAND( Auto, "Auto", "Preview what is active on the current window", EUserInterfaceActionType::RadioButton, FInputChord() );
+	UI_COMMAND( Lock, "Lock", "Lock to current preview", EUserInterfaceActionType::RadioButton, FInputChord() );
 
-	UI_COMMAND( CameraFollow, "Camera Follow", "Follow the bound of the mesh ", EUserInterfaceActionType::ToggleButton, FInputGesture() );
+	UI_COMMAND( CameraFollow, "Camera Follow", "Follow the bound of the mesh ", EUserInterfaceActionType::ToggleButton, FInputChord() );
 
-	UI_COMMAND( UseInGameBound, "In-game Bound", "Use in-game bound on preview mesh", EUserInterfaceActionType::ToggleButton, FInputGesture() );
+	UI_COMMAND( UseInGameBound, "In-game Bound", "Use in-game bound on preview mesh", EUserInterfaceActionType::ToggleButton, FInputChord() );
 
-	UI_COMMAND( SetShowNormals, "Normals", "Toggles display of vertex normals in the Preview Pane.", EUserInterfaceActionType::ToggleButton, FInputGesture() );
-	UI_COMMAND( SetShowTangents, "Tangents", "Toggles display of vertex tangents in the Preview Pane.", EUserInterfaceActionType::ToggleButton, FInputGesture() );
-	UI_COMMAND( SetShowBinormals, "Binormals", "Toggles display of vertex binormals (orthogonal vector to normal and tangent) in the Preview Pane.", EUserInterfaceActionType::ToggleButton, FInputGesture() );
+	UI_COMMAND( SetShowNormals, "Normals", "Toggles display of vertex normals in the Preview Pane.", EUserInterfaceActionType::ToggleButton, FInputChord() );
+	UI_COMMAND( SetShowTangents, "Tangents", "Toggles display of vertex tangents in the Preview Pane.", EUserInterfaceActionType::ToggleButton, FInputChord() );
+	UI_COMMAND( SetShowBinormals, "Binormals", "Toggles display of vertex binormals (orthogonal vector to normal and tangent) in the Preview Pane.", EUserInterfaceActionType::ToggleButton, FInputChord() );
 
-	UI_COMMAND( AnimSetDrawUVs, "UV", "Toggles display of the mesh's UVs for the specified channel.", EUserInterfaceActionType::ToggleButton, FInputGesture() );
+	UI_COMMAND( AnimSetDrawUVs, "UV", "Toggles display of the mesh's UVs for the specified channel.", EUserInterfaceActionType::ToggleButton, FInputChord() );
 }
 
 #undef LOCTEXT_NAMESPACE
\ No newline at end of file
diff --git a/Engine/Source/Editor/Persona/Private/AnimViewportPlaybackCommands.cpp b/Engine/Source/Editor/Persona/Private/AnimViewportPlaybackCommands.cpp
--- a/Engine/Source/Editor/Persona/Private/AnimViewportPlaybackCommands.cpp
+++ b/Engine/Source/Editor/Persona/Private/AnimViewportPlaybackCommands.cpp
@@ -19,13 +19,13 @@ FAnimViewportPlaybackCommands::FAnimViewportPlaybackCommands() : TCommands<FAnim
 
 void FAnimViewportPlaybackCommands::RegisterCommands()
 {
-	UI_COMMAND( PlaybackSpeedCommands[EAnimationPlaybackSpeeds::OneTenth],	"x0.1", "Set the animation playback speed to a tenth of normal", EUserInterfaceActionType::RadioButton, FInputGesture() );
-	UI_COMMAND( PlaybackSpeedCommands[EAnimationPlaybackSpeeds::Quarter],		"x0.25", "Set the animation playback speed to a quarter of normal", EUserInterfaceActionType::RadioButton, FInputGesture() );
-	UI_COMMAND( PlaybackSpeedCommands[EAnimationPlaybackSpeeds::Half],		"x0.5", "Set the animation playback speed to a half of normal", EUserInterfaceActionType::RadioButton, FInputGesture() );
-	UI_COMMAND( PlaybackSpeedCommands[EAnimationPlaybackSpeeds::Normal],		"x1.0", "Set the animation playback speed to normal", EUserInterfaceActionType::RadioButton, FInputGesture() );
-	UI_COMMAND( PlaybackSpeedCommands[EAnimationPlaybackSpeeds::Double],		"x2.0", "Set the animation playback speed to double the speed of normal", EUserInterfaceActionType::RadioButton, FInputGesture() );
-	UI_COMMAND( PlaybackSpeedCommands[EAnimationPlaybackSpeeds::FiveTimes],	"x5.0", "Set the animation playback speed to five times the normal speed", EUserInterfaceActionType::RadioButton, FInputGesture() );
-	UI_COMMAND( PlaybackSpeedCommands[EAnimationPlaybackSpeeds::TenTimes],	"x10.0", "Set the animation playback speed to ten times the normal speed", EUserInterfaceActionType::RadioButton, FInputGesture() );
+	UI_COMMAND( PlaybackSpeedCommands[EAnimationPlaybackSpeeds::OneTenth],	"x0.1", "Set the animation playback speed to a tenth of normal", EUserInterfaceActionType::RadioButton, FInputChord() );
+	UI_COMMAND( PlaybackSpeedCommands[EAnimationPlaybackSpeeds::Quarter],		"x0.25", "Set the animation playback speed to a quarter of normal", EUserInterfaceActionType::RadioButton, FInputChord() );
+	UI_COMMAND( PlaybackSpeedCommands[EAnimationPlaybackSpeeds::Half],		"x0.5", "Set the animation playback speed to a half of normal", EUserInterfaceActionType::RadioButton, FInputChord() );
+	UI_COMMAND( PlaybackSpeedCommands[EAnimationPlaybackSpeeds::Normal],		"x1.0", "Set the animation playback speed to normal", EUserInterfaceActionType::RadioButton, FInputChord() );
+	UI_COMMAND( PlaybackSpeedCommands[EAnimationPlaybackSpeeds::Double],		"x2.0", "Set the animation playback speed to double the speed of normal", EUserInterfaceActionType::RadioButton, FInputChord() );
+	UI_COMMAND( PlaybackSpeedCommands[EAnimationPlaybackSpeeds::FiveTimes],	"x5.0", "Set the animation playback speed to five times the normal speed", EUserInterfaceActionType::RadioButton, FInputChord() );
+	UI_COMMAND( PlaybackSpeedCommands[EAnimationPlaybackSpeeds::TenTimes],	"x10.0", "Set the animation playback speed to ten times the normal speed", EUserInterfaceActionType::RadioButton, FInputChord() );
 }
 
 #undef LOCTEXT_NAMESPACE
\ No newline at end of file
diff --git a/Engine/Source/Editor/Persona/Private/AnimViewportShowCommands.cpp b/Engine/Source/Editor/Persona/Private/AnimViewportShowCommands.cpp
--- a/Engine/Source/Editor/Persona/Private/AnimViewportShowCommands.cpp
+++ b/Engine/Source/Editor/Persona/Private/AnimViewportShowCommands.cpp
@@ -7,59 +7,59 @@
 
 void FAnimViewportShowCommands::RegisterCommands()
 {
-	UI_COMMAND( ToggleGrid, "Grid", "Display Grid", EUserInterfaceActionType::ToggleButton, FInputGesture() );
+	UI_COMMAND( ToggleGrid, "Grid", "Display Grid", EUserInterfaceActionType::ToggleButton, FInputChord() );
 
-	UI_COMMAND( ToggleFloor, "Floor", "Display Floor", EUserInterfaceActionType::ToggleButton, FInputGesture() );
-	UI_COMMAND( ToggleSky, "Sky", "Display Sky", EUserInterfaceActionType::ToggleButton, FInputGesture() );
+	UI_COMMAND( ToggleFloor, "Floor", "Display Floor", EUserInterfaceActionType::ToggleButton, FInputChord() );
+	UI_COMMAND( ToggleSky, "Sky", "Display Sky", EUserInterfaceActionType::ToggleButton, FInputChord() );
 
-	UI_COMMAND( MuteAudio, "Mute Audio", "Mutes audio from the preview", EUserInterfaceActionType::ToggleButton, FInputGesture() );
+	UI_COMMAND( MuteAudio, "Mute Audio", "Mutes audio from the preview", EUserInterfaceActionType::ToggleButton, FInputChord() );
 
-	UI_COMMAND(ProcessRootMotion, "Process Root Motion", "Moves preview based on animation root motion", EUserInterfaceActionType::ToggleButton, FInputGesture());
+	UI_COMMAND(ProcessRootMotion, "Process Root Motion", "Moves preview based on animation root motion", EUserInterfaceActionType::ToggleButton, FInputChord());
 
-	UI_COMMAND( ShowReferencePose, "Reference Pose", "Show reference pose on preview mesh", EUserInterfaceActionType::ToggleButton, FInputGesture() );
-	UI_COMMAND( ShowRetargetBasePose, "Retarget Base Pose", "Show retarget Base pose on preview mesh", EUserInterfaceActionType::ToggleButton, FInputGesture() );
-	UI_COMMAND( ShowBound, "Bound", "Show bound on preview mesh", EUserInterfaceActionType::ToggleButton, FInputGesture() );
-	UI_COMMAND( ShowPreviewMesh, "Mesh", "Show the preview mesh", EUserInterfaceActionType::ToggleButton, FInputGesture() );
-	UI_COMMAND( ShowMorphTargets, "MorphTargets", "Display Applied Morph Targets of the mesh", EUserInterfaceActionType::ToggleButton, FInputGesture() );
+	UI_COMMAND( ShowReferencePose, "Reference Pose", "Show reference pose on preview mesh", EUserInterfaceActionType::ToggleButton, FInputChord() );
+	UI_COMMAND( ShowRetargetBasePose, "Retarget Base Pose", "Show retarget Base pose on preview mesh", EUserInterfaceActionType::ToggleButton, FInputChord() );
+	UI_COMMAND( ShowBound, "Bound", "Show bound on preview mesh", EUserInterfaceActionType::ToggleButton, FInputChord() );
+	UI_COMMAND( ShowPreviewMesh, "Mesh", "Show the preview mesh", EUserInterfaceActionType::ToggleButton, FInputChord() );
+	UI_COMMAND( ShowMorphTargets, "MorphTargets", "Display Applied Morph Targets of the mesh", EUserInterfaceActionType::ToggleButton, FInputChord() );
 
-	UI_COMMAND( ShowBones, "Bones", "Display Skeleton in viewport", EUserInterfaceActionType::ToggleButton, FInputGesture() );
-	UI_COMMAND( ShowBoneNames, "Bone Names", "Display Bone Names in Viewport", EUserInterfaceActionType::ToggleButton, FInputGesture() );
+	UI_COMMAND( ShowBones, "Bones", "Display Skeleton in viewport", EUserInterfaceActionType::ToggleButton, FInputChord() );
+	UI_COMMAND( ShowBoneNames, "Bone Names", "Display Bone Names in Viewport", EUserInterfaceActionType::ToggleButton, FInputChord() );
 
 	// below 3 menus are radio button styles
-	UI_COMMAND(ShowDisplayInfoBasic, "Basic", "Display Basic Mesh Info in Viewport", EUserInterfaceActionType::RadioButton, FInputGesture());
-	UI_COMMAND(ShowDisplayInfoDetailed, "Detailed", "Display Detailed Mesh Info in Viewport", EUserInterfaceActionType::RadioButton, FInputGesture());
-	UI_COMMAND(HideDisplayInfo, "None", "Hide All Display Info in Viewport", EUserInterfaceActionType::RadioButton, FInputGesture());
-
-	UI_COMMAND( ShowBoneWeight, "View Selected Bone Weight", "Display Selected Bone Weight in Viewport", EUserInterfaceActionType::ToggleButton, FInputGesture() );
-	UI_COMMAND( ShowRawAnimation, "Uncompressed Animation", "Display Skeleton With Uncompressed Animation Data", EUserInterfaceActionType::ToggleButton, FInputGesture() );
-	UI_COMMAND( ShowNonRetargetedAnimation, "NonRetargeted Animation", "Display Skeleton With non retargeted Animation Data", EUserInterfaceActionType::ToggleButton, FInputGesture() );
-	UI_COMMAND( ShowAdditiveBaseBones, "Additive Base", "Display Skeleton In Additive Base Pose", EUserInterfaceActionType::ToggleButton, FInputGesture() );
-	UI_COMMAND( ShowSourceRawAnimation, "Source Animation", "Display Skeleton In Source Raw Animation if you have Track Curves Modified.", EUserInterfaceActionType::ToggleButton, FInputGesture() );
-	UI_COMMAND( ShowBakedAnimation, "Baked Animation", "Display Skeleton In Baked Raw Animation if you have Track Curves Modified.", EUserInterfaceActionType::ToggleButton, FInputGesture() );
-	UI_COMMAND( ShowSockets, "Sockets", "Display socket hitpoints", EUserInterfaceActionType::ToggleButton, FInputGesture() );
-
-	UI_COMMAND( ShowLocalAxesNone, "None", "Hides all local hierarchy axis", EUserInterfaceActionType::RadioButton, FInputGesture() );
-	UI_COMMAND( ShowLocalAxesSelected, "Selected Hierarchy", "Shows only the local bone axis of the selected bones", EUserInterfaceActionType::RadioButton, FInputGesture() );
-	UI_COMMAND( ShowLocalAxesAll, "All", "Shows all local hierarchy axes", EUserInterfaceActionType::RadioButton, FInputGesture() );
+	UI_COMMAND(ShowDisplayInfoBasic, "Basic", "Display Basic Mesh Info in Viewport", EUserInterfaceActionType::RadioButton, FInputChord());
+	UI_COMMAND(ShowDisplayInfoDetailed, "Detailed", "Display Detailed Mesh Info in Viewport", EUserInterfaceActionType::RadioButton, FInputChord());
+	UI_COMMAND(HideDisplayInfo, "None", "Hide All Display Info in Viewport", EUserInterfaceActionType::RadioButton, FInputChord());
+
+	UI_COMMAND( ShowBoneWeight, "View Selected Bone Weight", "Display Selected Bone Weight in Viewport", EUserInterfaceActionType::ToggleButton, FInputChord() );
+	UI_COMMAND( ShowRawAnimation, "Uncompressed Animation", "Display Skeleton With Uncompressed Animation Data", EUserInterfaceActionType::ToggleButton, FInputChord() );
+	UI_COMMAND( ShowNonRetargetedAnimation, "NonRetargeted Animation", "Display Skeleton With non retargeted Animation Data", EUserInterfaceActionType::ToggleButton, FInputChord() );
+	UI_COMMAND( ShowAdditiveBaseBones, "Additive Base", "Display Skeleton In Additive Base Pose", EUserInterfaceActionType::ToggleButton, FInputChord() );
+	UI_COMMAND( ShowSourceRawAnimation, "Source Animation", "Display Skeleton In Source Raw Animation if you have Track Curves Modified.", EUserInterfaceActionType::ToggleButton, FInputChord() );
+	UI_COMMAND( ShowBakedAnimation, "Baked Animation", "Display Skeleton In Baked Raw Animation if you have Track Curves Modified.", EUserInterfaceActionType::ToggleButton, FInputChord() );
+	UI_COMMAND( ShowSockets, "Sockets", "Display socket hitpoints", EUserInterfaceActionType::ToggleButton, FInputChord() );
+
+	UI_COMMAND( ShowLocalAxesNone, "None", "Hides all local hierarchy axis", EUserInterfaceActionType::RadioButton, FInputChord() );
+	UI_COMMAND( ShowLocalAxesSelected, "Selected Hierarchy", "Shows only the local bone axis of the selected bones", EUserInterfaceActionType::RadioButton, FInputChord() );
+	UI_COMMAND( ShowLocalAxesAll, "All", "Shows all local hierarchy axes", EUserInterfaceActionType::RadioButton, FInputChord() );
 
 #if WITH_APEX_CLOTHING
-	UI_COMMAND( DisableClothSimulation, "Disable Cloth Simulation", "Disable Cloth Simulation and Show non-simulated mesh", EUserInterfaceActionType::ToggleButton, FInputGesture() );
-	UI_COMMAND( ApplyClothWind, "Apply Wind", "Apply Wind to Clothing", EUserInterfaceActionType::ToggleButton, FInputGesture() );
+	UI_COMMAND( DisableClothSimulation, "Disable Cloth Simulation", "Disable Cloth Simulation and Show non-simulated mesh", EUserInterfaceActionType::ToggleButton, FInputChord() );
+	UI_COMMAND( ApplyClothWind, "Apply Wind", "Apply Wind to Clothing", EUserInterfaceActionType::ToggleButton, FInputChord() );
 
-	UI_COMMAND( ShowClothSimulationNormals, "Simulation Normals", "Display Cloth Simulation Normals", EUserInterfaceActionType::ToggleButton, FInputGesture() );
-	UI_COMMAND( ShowClothGraphicalTangents, "Graphical Tangents", "Display Cloth Graphical Tangents", EUserInterfaceActionType::ToggleButton, FInputGesture() );
-	UI_COMMAND( ShowClothCollisionVolumes, "Collision Volumes", "Display Cloth Collision Volumes", EUserInterfaceActionType::ToggleButton, FInputGesture() );
-	UI_COMMAND( EnableCollisionWithAttachedClothChildren, "Collide with Cloth Children", "Enables collision detection between collision primitives in the base mesh and clothing on any attachments in the preview scene.", EUserInterfaceActionType::ToggleButton, FInputGesture() );	
+	UI_COMMAND( ShowClothSimulationNormals, "Simulation Normals", "Display Cloth Simulation Normals", EUserInterfaceActionType::ToggleButton, FInputChord() );
+	UI_COMMAND( ShowClothGraphicalTangents, "Graphical Tangents", "Display Cloth Graphical Tangents", EUserInterfaceActionType::ToggleButton, FInputChord() );
+	UI_COMMAND( ShowClothCollisionVolumes, "Collision Volumes", "Display Cloth Collision Volumes", EUserInterfaceActionType::ToggleButton, FInputChord() );
+	UI_COMMAND( EnableCollisionWithAttachedClothChildren, "Collide with Cloth Children", "Enables collision detection between collision primitives in the base mesh and clothing on any attachments in the preview scene.", EUserInterfaceActionType::ToggleButton, FInputChord() );	
 
-	UI_COMMAND( ShowClothPhysicalMeshWire, "Physical Mesh Wire", "Display Simulation Mesh's wire frame. Shows all physical mesh loaded from an asset file when cloth simulation is disabled and shows simulated results when simulating. The part having zero max distance value will be drawn in pink.", EUserInterfaceActionType::ToggleButton, FInputGesture() );
-	UI_COMMAND( ShowClothMaxDistances, "Max Distances", "Display Max Distances only when simulation is disabled. If turned on this option, disables cloth simulation automatically", EUserInterfaceActionType::ToggleButton, FInputGesture() );
-	UI_COMMAND( ShowClothBackstop, "Back stops", "Display Back stops. If it has a proper value, draws a line in red. Otherwise, draws a vertex in white.", EUserInterfaceActionType::ToggleButton, FInputGesture() );
-	UI_COMMAND( ShowClothFixedVertices, "Fixed Vertices", "Display Fixed Vertices", EUserInterfaceActionType::ToggleButton, FInputGesture() );	
+	UI_COMMAND( ShowClothPhysicalMeshWire, "Physical Mesh Wire", "Display Simulation Mesh's wire frame. Shows all physical mesh loaded from an asset file when cloth simulation is disabled and shows simulated results when simulating. The part having zero max distance value will be drawn in pink.", EUserInterfaceActionType::ToggleButton, FInputChord() );
+	UI_COMMAND( ShowClothMaxDistances, "Max Distances", "Display Max Distances only when simulation is disabled. If turned on this option, disables cloth simulation automatically", EUserInterfaceActionType::ToggleButton, FInputChord() );
+	UI_COMMAND( ShowClothBackstop, "Back stops", "Display Back stops. If it has a proper value, draws a line in red. Otherwise, draws a vertex in white.", EUserInterfaceActionType::ToggleButton, FInputChord() );
+	UI_COMMAND( ShowClothFixedVertices, "Fixed Vertices", "Display Fixed Vertices", EUserInterfaceActionType::ToggleButton, FInputChord() );	
 
 	// below 3 menus are radio button styles
-	UI_COMMAND(ShowAllSections, "Show All Sections", "Display All sections including Cloth Mapped Sections", EUserInterfaceActionType::RadioButton, FInputGesture());
-	UI_COMMAND(ShowOnlyClothSections, "Show Only Cloth Sections", "Display Only Cloth Mapped Sections", EUserInterfaceActionType::RadioButton, FInputGesture());
-	UI_COMMAND(HideOnlyClothSections, "Hide Only Cloth Sections", "Display All Except Cloth Mapped Sections", EUserInterfaceActionType::RadioButton, FInputGesture());
+	UI_COMMAND(ShowAllSections, "Show All Sections", "Display All sections including Cloth Mapped Sections", EUserInterfaceActionType::RadioButton, FInputChord());
+	UI_COMMAND(ShowOnlyClothSections, "Show Only Cloth Sections", "Display Only Cloth Mapped Sections", EUserInterfaceActionType::RadioButton, FInputChord());
+	UI_COMMAND(HideOnlyClothSections, "Hide Only Cloth Sections", "Display All Except Cloth Mapped Sections", EUserInterfaceActionType::RadioButton, FInputChord());
 
 #endif// #if WITH_APEX_CLOTHING
 }
diff --git a/Engine/Source/Editor/Persona/Private/PersonaCommands.cpp b/Engine/Source/Editor/Persona/Private/PersonaCommands.cpp
--- a/Engine/Source/Editor/Persona/Private/PersonaCommands.cpp
+++ b/Engine/Source/Editor/Persona/Private/PersonaCommands.cpp
@@ -9,32 +9,32 @@
 void FPersonaCommands::RegisterCommands()
 {
 	// save menu
-	UI_COMMAND(SaveAnimationAssets, "Save Animation Assets", "Only save all animation assets including skeletons, skeletal meshes and animation blueprints.", EUserInterfaceActionType::Button, FInputGesture());
+	UI_COMMAND(SaveAnimationAssets, "Save Animation Assets", "Only save all animation assets including skeletons, skeletal meshes and animation blueprints.", EUserInterfaceActionType::Button, FInputChord());
 
 	// skeleton menu
-	UI_COMMAND( ChangeSkeletonPreviewMesh, "Set Preview Mesh as Default", "Changes the skeletons default preview mesh to the current open preview mesh. The skeleton will require saving after this action.", EUserInterfaceActionType::Button, FInputGesture() );
-	UI_COMMAND( RemoveUnusedBones, "Remove Unused Bones from Skeleton", "Removes any bones from the skeleton that are not used by any of its meshes. The skeleton and associated animations will require saving after this action.", EUserInterfaceActionType::Button, FInputGesture() );
-	UI_COMMAND( AnimNotifyWindow, "Anim Notifies", "You can manage animation notifies that belong to the skeleton.", EUserInterfaceActionType::Button, FInputGesture() );
-	UI_COMMAND( RetargetManager, "Retarget Manager", "Manager retarget setups. ", EUserInterfaceActionType::Button, FInputGesture() );
-	UI_COMMAND( ImportMesh, "Import Mesh", "Import new mesh for this skeleton. ", EUserInterfaceActionType::Button, FInputGesture() );
+	UI_COMMAND( ChangeSkeletonPreviewMesh, "Set Preview Mesh as Default", "Changes the skeletons default preview mesh to the current open preview mesh. The skeleton will require saving after this action.", EUserInterfaceActionType::Button, FInputChord() );
+	UI_COMMAND( RemoveUnusedBones, "Remove Unused Bones from Skeleton", "Removes any bones from the skeleton that are not used by any of its meshes. The skeleton and associated animations will require saving after this action.", EUserInterfaceActionType::Button, FInputChord() );
+	UI_COMMAND( AnimNotifyWindow, "Anim Notifies", "You can manage animation notifies that belong to the skeleton.", EUserInterfaceActionType::Button, FInputChord() );
+	UI_COMMAND( RetargetManager, "Retarget Manager", "Manager retarget setups. ", EUserInterfaceActionType::Button, FInputChord() );
+	UI_COMMAND( ImportMesh, "Import Mesh", "Import new mesh for this skeleton. ", EUserInterfaceActionType::Button, FInputChord() );
 
 	// mesh menu
-	UI_COMMAND( ReimportMesh, "Reimport Mesh", "Reimport the current mesh.", EUserInterfaceActionType::Button, FInputGesture() );
-	UI_COMMAND( ImportLODs, "Import LODs", "Import LODs of the current mesh.", EUserInterfaceActionType::Button, FInputGesture() );
-	UI_COMMAND( AddBodyPart, "Add Body Part", "Add extra body parts for the current mesh", EUserInterfaceActionType::Button, FInputGesture());
+	UI_COMMAND( ReimportMesh, "Reimport Mesh", "Reimport the current mesh.", EUserInterfaceActionType::Button, FInputChord() );
+	UI_COMMAND( ImportLODs, "Import LODs", "Import LODs of the current mesh.", EUserInterfaceActionType::Button, FInputChord() );
+	UI_COMMAND( AddBodyPart, "Add Body Part", "Add extra body parts for the current mesh", EUserInterfaceActionType::Button, FInputChord());
 
 	// animation menu
 	// import animation
-	UI_COMMAND( ImportAnimation, "Import Animation", "Import new animation for the skeleton.", EUserInterfaceActionType::Button, FInputGesture());
-	UI_COMMAND( ReimportAnimation, "Reimport Animation", "Reimport current animation.", EUserInterfaceActionType::Button, FInputGesture());
-	UI_COMMAND( ApplyCompression, "Apply Compression", "Apply compression to current animation", EUserInterfaceActionType::Button, FInputGesture() );
-	UI_COMMAND( RecordAnimation, "Record to new Animation", "Create new animation from currently playing", EUserInterfaceActionType::Button, FInputGesture() );
-	UI_COMMAND( ExportToFBX, "Export to FBX", "Export current animation to FBX", EUserInterfaceActionType::Button, FInputGesture() );
-	UI_COMMAND( AddLoopingInterpolation, "Add Looping Interpolation", "Add an extra first frame at the end of the animation to create interpolation when looping", EUserInterfaceActionType::Button, FInputGesture() );
-	UI_COMMAND( SetKey, "Set Key", "Add Bone Transform to Additive Layer Tracks", EUserInterfaceActionType::Button, FInputGesture(EKeys::S) );
-	UI_COMMAND( ApplyAnimation, "Apply Animation", "Apply Additive Layer Tracks to Runtime Animation Data", EUserInterfaceActionType::Button, FInputGesture() );
-
-	UI_COMMAND( UpdateSkeletonRefPose, "Update Skeleton RefPose", "Update Skeleton ref pose based on current preview mesh", EUserInterfaceActionType::Button, FInputGesture() );
+	UI_COMMAND( ImportAnimation, "Import Animation", "Import new animation for the skeleton.", EUserInterfaceActionType::Button, FInputChord());
+	UI_COMMAND( ReimportAnimation, "Reimport Animation", "Reimport current animation.", EUserInterfaceActionType::Button, FInputChord());
+	UI_COMMAND( ApplyCompression, "Apply Compression", "Apply compression to current animation", EUserInterfaceActionType::Button, FInputChord() );
+	UI_COMMAND( RecordAnimation, "Record to new Animation", "Create new animation from currently playing", EUserInterfaceActionType::Button, FInputChord() );
+	UI_COMMAND( ExportToFBX, "Export to FBX", "Export current animation to FBX", EUserInterfaceActionType::Button, FInputChord() );
+	UI_COMMAND( AddLoopingInterpolation, "Add Looping Interpolation", "Add an extra first frame at the end of the animation to create interpolation when looping", EUserInterfaceActionType::Button, FInputChord() );
+	UI_COMMAND( SetKey, "Set Key", "Add Bone Transform to Additive Layer Tracks", EUserInterfaceActionType::Button, FInputChord(EKeys::S) );
+	UI_COMMAND( ApplyAnimation, "Apply Animation", "Apply Additive Layer Tracks to Runtime Animation Data", EUserInterfaceActionType::Button, FInputChord() );
+
+	UI_COMMAND( UpdateSkeletonRefPose, "Update Skeleton RefPose", "Update Skeleton ref pose based on current preview mesh", EUserInterfaceActionType::Button, FInputChord() );
 }
 
 #undef LOCTEXT_NAMESPACE
\ No newline at end of file
diff --git a/Engine/Source/Editor/Persona/Private/SAnimNotifyPanel.cpp b/Engine/Source/Editor/Persona/Private/SAnimNotifyPanel.cpp
--- a/Engine/Source/Editor/Persona/Private/SAnimNotifyPanel.cpp
+++ b/Engine/Source/Editor/Persona/Private/SAnimNotifyPanel.cpp
@@ -912,7 +912,7 @@ public:
 
 	virtual void RegisterCommands() override
 	{
-		UI_COMMAND( SomeSequenceAction, "Some Sequence Action", "Does some sequence action", EUserInterfaceActionType::Button, FInputGesture() );
+		UI_COMMAND( SomeSequenceAction, "Some Sequence Action", "Does some sequence action", EUserInterfaceActionType::Button, FInputChord() );
 	}
 public:
 	TSharedPtr<FUICommandInfo> SomeSequenceAction;
@@ -3037,7 +3037,7 @@ bool SNotifyEdTrack::CanDeleteTrack()
 
 void FAnimNotifyPanelCommands::RegisterCommands()
 {
-	UI_COMMAND(DeleteNotify, "Delete", "Deletes the selected notifies.", EUserInterfaceActionType::Button, FInputGesture(EKeys::Platform_Delete));
+	UI_COMMAND(DeleteNotify, "Delete", "Deletes the selected notifies.", EUserInterfaceActionType::Button, FInputChord(EKeys::Platform_Delete));
 }
 
 //////////////////////////////////////////////////////////////////////////
diff --git a/Engine/Source/Editor/Persona/Private/SkeletonTreeCommands.cpp b/Engine/Source/Editor/Persona/Private/SkeletonTreeCommands.cpp
--- a/Engine/Source/Editor/Persona/Private/SkeletonTreeCommands.cpp
+++ b/Engine/Source/Editor/Persona/Private/SkeletonTreeCommands.cpp
@@ -7,28 +7,28 @@
 
 void FSkeletonTreeCommands::RegisterCommands()
 {
-	UI_COMMAND( ShowAllBones, "Show All Bones", "Show every bone in the skeleton", EUserInterfaceActionType::RadioButton, FInputGesture() );
-	UI_COMMAND( ShowMeshBones, "Show Mesh Bones", "Show bones that are used in the mesh", EUserInterfaceActionType::RadioButton, FInputGesture() );
-	UI_COMMAND( ShowWeightedBones, "Show Weighted Bones", "Show bones that have vertices weighted to them", EUserInterfaceActionType::RadioButton, FInputGesture() );
-	UI_COMMAND( HideBones, "Hide Bones", "Hides all bones (sockets and attached assets will still be listed)", EUserInterfaceActionType::RadioButton, FInputGesture() );
-
-	UI_COMMAND( CopyBoneNames, "Copy Selected Bone Names", "Copy selected bone names to clipboard", EUserInterfaceActionType::Button, FInputGesture() );
-	UI_COMMAND( ResetBoneTransforms, "Reset Selected Bone Transforms", "Reset the transforms of the selected bones", EUserInterfaceActionType::Button, FInputGesture() );
-	UI_COMMAND( CopySockets, "Copy Selected Sockets", "Copy selected sockets to clipboard", EUserInterfaceActionType::Button, FInputGesture( EModifierKey::Control, EKeys::C ) );
-	UI_COMMAND( PasteSockets, "Paste Selected Sockets", "Paste selected sockets from clipboard", EUserInterfaceActionType::Button, FInputGesture( EModifierKey::Control, EKeys::V ) );
-
-	UI_COMMAND( AddSocket, "Add Socket", "Add a socket to this bone in the skeleton (disabled when socket filter is set to \"Mesh Sockets\" or \"Sockets Hidden\" mode)", EUserInterfaceActionType::Button, FInputGesture() );
-	UI_COMMAND( CreateMeshSocket, "Create Mesh Socket", "Duplicate this socket from skeleton to the current mesh and modify the socket data for it", EUserInterfaceActionType::Button, FInputGesture() );
-	UI_COMMAND( RemoveMeshSocket, "Remove Mesh Socket", "Remove duplicated version of this socket for the current mesh (reverts to the socket in the skeleton)", EUserInterfaceActionType::Button, FInputGesture() );
-	UI_COMMAND( PromoteSocketToSkeleton, "Promote Socket To Skeleton", "Makes this socket available for all meshes that use the same skeleton (copies the socket from this mesh to the skeleton)", EUserInterfaceActionType::Button, FInputGesture() );
-
-	UI_COMMAND(DeleteSelectedRows, "Delete", "Delete all selected sockets and attached meshes in the tree", EUserInterfaceActionType::Button, FInputGesture(EKeys::Platform_Delete));
-
-	UI_COMMAND( ShowActiveSockets, "Show Active Sockets", "Show mesh and skeleton sockets, hiding the skeleton sockets that have a customized mesh socket", EUserInterfaceActionType::RadioButton, FInputGesture() );
-	UI_COMMAND( ShowAllSockets, "Show All Sockets", "Show all sockets that are in the mesh and skeleton", EUserInterfaceActionType::RadioButton, FInputGesture() );
-	UI_COMMAND( ShowMeshSockets, "Show Mesh Sockets", "Show sockets that are in the mesh only", EUserInterfaceActionType::RadioButton, FInputGesture() );
-	UI_COMMAND( ShowSkeletonSockets, "Show Skeleton Sockets", "Show sockets that are in the skeleton only", EUserInterfaceActionType::RadioButton, FInputGesture() );
-	UI_COMMAND( HideSockets, "Hide Sockets", "Show no sockets", EUserInterfaceActionType::RadioButton, FInputGesture() );
+	UI_COMMAND( ShowAllBones, "Show All Bones", "Show every bone in the skeleton", EUserInterfaceActionType::RadioButton, FInputChord() );
+	UI_COMMAND( ShowMeshBones, "Show Mesh Bones", "Show bones that are used in the mesh", EUserInterfaceActionType::RadioButton, FInputChord() );
+	UI_COMMAND( ShowWeightedBones, "Show Weighted Bones", "Show bones that have vertices weighted to them", EUserInterfaceActionType::RadioButton, FInputChord() );
+	UI_COMMAND( HideBones, "Hide Bones", "Hides all bones (sockets and attached assets will still be listed)", EUserInterfaceActionType::RadioButton, FInputChord() );
+
+	UI_COMMAND( CopyBoneNames, "Copy Selected Bone Names", "Copy selected bone names to clipboard", EUserInterfaceActionType::Button, FInputChord() );
+	UI_COMMAND( ResetBoneTransforms, "Reset Selected Bone Transforms", "Reset the transforms of the selected bones", EUserInterfaceActionType::Button, FInputChord() );
+	UI_COMMAND( CopySockets, "Copy Selected Sockets", "Copy selected sockets to clipboard", EUserInterfaceActionType::Button, FInputChord( EModifierKey::Control, EKeys::C ) );
+	UI_COMMAND( PasteSockets, "Paste Selected Sockets", "Paste selected sockets from clipboard", EUserInterfaceActionType::Button, FInputChord( EModifierKey::Control, EKeys::V ) );
+
+	UI_COMMAND( AddSocket, "Add Socket", "Add a socket to this bone in the skeleton (disabled when socket filter is set to \"Mesh Sockets\" or \"Sockets Hidden\" mode)", EUserInterfaceActionType::Button, FInputChord() );
+	UI_COMMAND( CreateMeshSocket, "Create Mesh Socket", "Duplicate this socket from skeleton to the current mesh and modify the socket data for it", EUserInterfaceActionType::Button, FInputChord() );
+	UI_COMMAND( RemoveMeshSocket, "Remove Mesh Socket", "Remove duplicated version of this socket for the current mesh (reverts to the socket in the skeleton)", EUserInterfaceActionType::Button, FInputChord() );
+	UI_COMMAND( PromoteSocketToSkeleton, "Promote Socket To Skeleton", "Makes this socket available for all meshes that use the same skeleton (copies the socket from this mesh to the skeleton)", EUserInterfaceActionType::Button, FInputChord() );
+
+	UI_COMMAND(DeleteSelectedRows, "Delete", "Delete all selected sockets and attached meshes in the tree", EUserInterfaceActionType::Button, FInputChord(EKeys::Platform_Delete));
+
+	UI_COMMAND( ShowActiveSockets, "Show Active Sockets", "Show mesh and skeleton sockets, hiding the skeleton sockets that have a customized mesh socket", EUserInterfaceActionType::RadioButton, FInputChord() );
+	UI_COMMAND( ShowAllSockets, "Show All Sockets", "Show all sockets that are in the mesh and skeleton", EUserInterfaceActionType::RadioButton, FInputChord() );
+	UI_COMMAND( ShowMeshSockets, "Show Mesh Sockets", "Show sockets that are in the mesh only", EUserInterfaceActionType::RadioButton, FInputChord() );
+	UI_COMMAND( ShowSkeletonSockets, "Show Skeleton Sockets", "Show sockets that are in the skeleton only", EUserInterfaceActionType::RadioButton, FInputChord() );
+	UI_COMMAND( HideSockets, "Hide Sockets", "Show no sockets", EUserInterfaceActionType::RadioButton, FInputChord() );
 }
 
 #undef LOCTEXT_NAMESPACE
\ No newline at end of file
diff --git a/Engine/Source/Editor/PhAT/Private/PhATActions.cpp b/Engine/Source/Editor/PhAT/Private/PhATActions.cpp
--- a/Engine/Source/Editor/PhAT/Private/PhATActions.cpp
+++ b/Engine/Source/Editor/PhAT/Private/PhATActions.cpp
@@ -8,88 +8,88 @@
 
 void FPhATCommands::RegisterCommands()
 {
-	UI_COMMAND(ChangeDefaultMesh, "Change Mesh", "Change Default SkeletalMesh", EUserInterfaceActionType::Button, FInputGesture());
-	UI_COMMAND(ResetEntireAsset, "Reset Asset", "Opens Body Creation Settings And Replaces Asset Using The New Settings", EUserInterfaceActionType::Button, FInputGesture());
-	UI_COMMAND(RestetBoneCollision, "Reset", "Reset Bone Collision", EUserInterfaceActionType::Button, FInputGesture());
-	UI_COMMAND(ApplyPhysicalMaterial, "Apply PhysMat", "Applies Currently Selected Physical Material In Content Browser To All Bodies", EUserInterfaceActionType::Button, FInputGesture());
-	UI_COMMAND(EditingMode_Body, "Body Mode", "Body Editing Mode", EUserInterfaceActionType::RadioButton, FInputGesture());
-	UI_COMMAND(EditingMode_Constraint, "Constraint Mode", "Constraint Editing Mode", EUserInterfaceActionType::RadioButton, FInputGesture());
-	UI_COMMAND(CopyProperties, "Copy Properties", "Copy Properties: Copy Properties Of Currently Selected Object To Next Selected Object", EUserInterfaceActionType::Button, FInputGesture(EModifierKey::Control, EKeys::C));
-	UI_COMMAND(PasteProperties, "Paste Properties", "Paste Properties: Copy Properties Of Currently Selected Object To Next Selected Object", EUserInterfaceActionType::Button, FInputGesture(EModifierKey::Control, EKeys::V));
-	UI_COMMAND(InstanceProperties, "Instance Properties", "Instance Properties: Displays Instance Properties When In Body Editing Mode", EUserInterfaceActionType::ToggleButton, FInputGesture(EKeys::I));
-	UI_COMMAND(RepeatLastSimulation, "Simulate", "Previews Physics Simulation", EUserInterfaceActionType::RadioButton, FInputGesture(EKeys::Enter));
-	UI_COMMAND(SimulationNormal, "Real Simulation", "Previews Normal Physics Simulation", EUserInterfaceActionType::RadioButton, FInputGesture());
-	UI_COMMAND(SimulationNoGravity, "No Gravity Simulation", "Run Physics Simulation without gravity. Use this to debug issues with your ragdoll. If the setup is correct, the asset should not move!", EUserInterfaceActionType::RadioButton, FInputGesture());
-	UI_COMMAND(ToggleSelectedSimulation, "Selected Simulation", "Run Physics Simulation on selected objects. Use this to tune  specific parts of your ragdoll.", EUserInterfaceActionType::ToggleButton, FInputGesture());
-	UI_COMMAND(MeshRenderingMode_Solid, "Solid", "Solid Mesh Rendering Mode", EUserInterfaceActionType::RadioButton, FInputGesture());
-	UI_COMMAND(MeshRenderingMode_Wireframe, "Wireframe", "Wireframe Mesh Rendering Mode", EUserInterfaceActionType::RadioButton, FInputGesture());
-	UI_COMMAND(MeshRenderingMode_None, "None", "No Mesh Rendering Mode", EUserInterfaceActionType::RadioButton, FInputGesture());
-	UI_COMMAND(CollisionRenderingMode_Solid, "Solid", "Solid Collision Rendering Mode", EUserInterfaceActionType::RadioButton, FInputGesture());
-	UI_COMMAND(CollisionRenderingMode_Wireframe, "Wireframe", "Wireframe Collision Rendering Mode", EUserInterfaceActionType::RadioButton, FInputGesture());
-	UI_COMMAND(CollisionRenderingMode_None, "None", "No Collision Rendering Mode", EUserInterfaceActionType::RadioButton, FInputGesture());
-	UI_COMMAND(ConstraintRenderingMode_None, "None", "No Constraint Rendering Mode", EUserInterfaceActionType::RadioButton, FInputGesture());
-	UI_COMMAND(ConstraintRenderingMode_AllPositions, "All Positions", "All Positions Constraint Rendering Mode", EUserInterfaceActionType::RadioButton, FInputGesture());
-	UI_COMMAND(ConstraintRenderingMode_AllLimits, "All Limits", "All Limits Constraint Rendering Mode", EUserInterfaceActionType::RadioButton, FInputGesture());
-	UI_COMMAND(ShowKinematicBodies, "Kinematic Bodies", "Displays Kinematic Bodies In Red", EUserInterfaceActionType::ToggleButton, FInputGesture());
-	UI_COMMAND(DrawGroundBox, "Ground Box", "Displays Floor Grid", EUserInterfaceActionType::ToggleButton, FInputGesture());
-	UI_COMMAND(ToggleGraphicsHierarchy, "Hierarchy", "Show Graphical Hierarchy Of Joints In Preview Viewport", EUserInterfaceActionType::ToggleButton, FInputGesture());
-	UI_COMMAND(ToggleBoneInfuences, "Bone Influences", "Displays Vertices Weighted To Currently Selected Bone or Body", EUserInterfaceActionType::ToggleButton, FInputGesture());
-	UI_COMMAND(ToggleMassProperties, "Mass Properties", "Show Mass Properties For Bodies When Simulation Is Enabled", EUserInterfaceActionType::ToggleButton, FInputGesture());
-	UI_COMMAND(DisableCollision, "Collision Off", "Disable Collision Of Currently Selected Body With Next Selected Body", EUserInterfaceActionType::Button, FInputGesture(EKeys::RightBracket));
-	UI_COMMAND(EnableCollision, "Collision On", "Enable Collision Of Currently Selected Body With Next Selected Body", EUserInterfaceActionType::Button, FInputGesture(EKeys::LeftBracket));
-	UI_COMMAND(WeldToBody, "Weld", "Weld Body: Weld Currently Selected Bodies", EUserInterfaceActionType::Button, FInputGesture());
-	UI_COMMAND(AddNewBody, "New Body", "Add New Body To Selected Bone.", EUserInterfaceActionType::Button, FInputGesture());
-	UI_COMMAND(AddSphere, "Add Sphere", "Add Sphere To Selected Bone", EUserInterfaceActionType::Button, FInputGesture());
-	UI_COMMAND(AddSphyl, "Add Sphyl", "Add Sphyl To Selected Bone", EUserInterfaceActionType::Button, FInputGesture());
-	UI_COMMAND(AddBox, "Add Box", "Add Box To Selected Bone", EUserInterfaceActionType::Button, FInputGesture());
-	UI_COMMAND(DeletePrimitive, "Delete", "Delete Selected Primitive", EUserInterfaceActionType::Button, FInputGesture());
-	UI_COMMAND(DuplicatePrimitive, "Duplicate", "Duplicate Selected Primitive", EUserInterfaceActionType::Button, FInputGesture());
-	UI_COMMAND(ResetConstraint, "Reset", "Reset Constraint", EUserInterfaceActionType::Button, FInputGesture());
-	UI_COMMAND(SnapConstraint, "Snap", "Snap Constraint To Bone", EUserInterfaceActionType::Button, FInputGesture());
-	UI_COMMAND(ConvertToBallAndSocket, "To B&S", "Convert Selected Constraint To Ball-And-Socket", EUserInterfaceActionType::Button, FInputGesture());
-	UI_COMMAND(ConvertToHinge, "To Hinge", "Convert Selected Constraint To Hinge", EUserInterfaceActionType::Button, FInputGesture());
-	UI_COMMAND(ConvertToPrismatic, "To Prismatic", "Convert Selected Constraint To Prismatic", EUserInterfaceActionType::Button, FInputGesture());
-	UI_COMMAND(ConvertToSkeletal, "To Skeletal", "Convert Selected Constraint To Skeletal", EUserInterfaceActionType::Button, FInputGesture());
-	UI_COMMAND(DeleteConstraint, "Delete", "Delete Selected Constraint", EUserInterfaceActionType::Button, FInputGesture());
-	UI_COMMAND(PlayAnimation, "Play", "Play Animation", EUserInterfaceActionType::ToggleButton, FInputGesture());
-	UI_COMMAND(RecordAnimation, "Record", "Record Animation", EUserInterfaceActionType::ToggleButton, FInputGesture());
-	UI_COMMAND(ShowSkeleton, "Skeleton", "Show Skeleton", EUserInterfaceActionType::ToggleButton, FInputGesture());
-	UI_COMMAND(MakeBodyKinematic, "Kinematic", "Make Body Kinematic", EUserInterfaceActionType::RadioButton, FInputGesture());
-	UI_COMMAND(MakeBodySimulated, "Simulated", "Make Body Simulated", EUserInterfaceActionType::RadioButton, FInputGesture());
-	UI_COMMAND(MakeBodyDefault, "Default", "Reset This Body To Default", EUserInterfaceActionType::RadioButton, FInputGesture());
-	UI_COMMAND(KinematicAllBodiesBelow, "Set All Bodies Below To Kinematic", "Set All Bodies Below To Kinematic", EUserInterfaceActionType::Button, FInputGesture());
-	UI_COMMAND(SimulatedAllBodiesBelow, "Set All Bodies Below To Simulated", "Set All Bodies Below To Simulated", EUserInterfaceActionType::Button, FInputGesture());
-	UI_COMMAND(MakeAllBodiesBelowDefault, "Reset All Bodies Below To Default", "Reset All Bodies Below To Default", EUserInterfaceActionType::Button, FInputGesture());
-	UI_COMMAND(DeleteBody, "Delete", "Delete Selected Body", EUserInterfaceActionType::Button, FInputGesture());
-	UI_COMMAND(DeleteAllBodiesBelow, "Delete All Bodies Below", "Delete All Bodies Below", EUserInterfaceActionType::Button, FInputGesture());
-	UI_COMMAND(ToggleMotor, "Toggle Motor", "Toggle Motor", EUserInterfaceActionType::Button, FInputGesture());
-	UI_COMMAND(EnableMotorsBelow, "Enable Motors Below", "Enable Motors Below", EUserInterfaceActionType::Button, FInputGesture());
-	UI_COMMAND(DisableMotorsBelow, "Disable Motors Below", "Disable Motors Below", EUserInterfaceActionType::Button, FInputGesture());
-	UI_COMMAND(SelectAllObjects, "Select All Objects", "Select All Objects", EUserInterfaceActionType::Button, FInputGesture(EModifierKey::Control, EKeys::A));
-	UI_COMMAND(HierarchyFilterAll, "All Bones", "Show Entire Hierarchy", EUserInterfaceActionType::Button, FInputGesture());
-	UI_COMMAND(HierarchyFilterBodies, "Bones With Bodies", "Filter Bones With Bodies", EUserInterfaceActionType::Button, FInputGesture());
-	UI_COMMAND(PerspectiveView, "Perspective", "Perspective", EUserInterfaceActionType::Button, FInputGesture());
-	UI_COMMAND(LeftView, "Left", "Orthographic view from left", EUserInterfaceActionType::Button, FInputGesture());
-	UI_COMMAND(TopView, "Top", "Orthographic view from top", EUserInterfaceActionType::Button, FInputGesture());
-	UI_COMMAND(FrontView, "Front", "Orthographic view from front", EUserInterfaceActionType::Button, FInputGesture());
-	UI_COMMAND(RightView, "Right", "Orthographic view from right", EUserInterfaceActionType::Button, FInputGesture());
-	UI_COMMAND(BottomView, "Bottom", "Orthographic view from bottom", EUserInterfaceActionType::Button, FInputGesture());
-	UI_COMMAND(BackView, "Back", "Orthographic view from back", EUserInterfaceActionType::Button, FInputGesture());
-	UI_COMMAND(Mirror, "Mirror", "Finds the body on the other side and duplicates constraint and body", EUserInterfaceActionType::Button, FInputGesture(EKeys::M));
+	UI_COMMAND(ChangeDefaultMesh, "Change Mesh", "Change Default SkeletalMesh", EUserInterfaceActionType::Button, FInputChord());
+	UI_COMMAND(ResetEntireAsset, "Reset Asset", "Opens Body Creation Settings And Replaces Asset Using The New Settings", EUserInterfaceActionType::Button, FInputChord());
+	UI_COMMAND(RestetBoneCollision, "Reset", "Reset Bone Collision", EUserInterfaceActionType::Button, FInputChord());
+	UI_COMMAND(ApplyPhysicalMaterial, "Apply PhysMat", "Applies Currently Selected Physical Material In Content Browser To All Bodies", EUserInterfaceActionType::Button, FInputChord());
+	UI_COMMAND(EditingMode_Body, "Body Mode", "Body Editing Mode", EUserInterfaceActionType::RadioButton, FInputChord());
+	UI_COMMAND(EditingMode_Constraint, "Constraint Mode", "Constraint Editing Mode", EUserInterfaceActionType::RadioButton, FInputChord());
+	UI_COMMAND(CopyProperties, "Copy Properties", "Copy Properties: Copy Properties Of Currently Selected Object To Next Selected Object", EUserInterfaceActionType::Button, FInputChord(EModifierKey::Control, EKeys::C));
+	UI_COMMAND(PasteProperties, "Paste Properties", "Paste Properties: Copy Properties Of Currently Selected Object To Next Selected Object", EUserInterfaceActionType::Button, FInputChord(EModifierKey::Control, EKeys::V));
+	UI_COMMAND(InstanceProperties, "Instance Properties", "Instance Properties: Displays Instance Properties When In Body Editing Mode", EUserInterfaceActionType::ToggleButton, FInputChord(EKeys::I));
+	UI_COMMAND(RepeatLastSimulation, "Simulate", "Previews Physics Simulation", EUserInterfaceActionType::RadioButton, FInputChord(EKeys::Enter));
+	UI_COMMAND(SimulationNormal, "Real Simulation", "Previews Normal Physics Simulation", EUserInterfaceActionType::RadioButton, FInputChord());
+	UI_COMMAND(SimulationNoGravity, "No Gravity Simulation", "Run Physics Simulation without gravity. Use this to debug issues with your ragdoll. If the setup is correct, the asset should not move!", EUserInterfaceActionType::RadioButton, FInputChord());
+	UI_COMMAND(ToggleSelectedSimulation, "Selected Simulation", "Run Physics Simulation on selected objects. Use this to tune  specific parts of your ragdoll.", EUserInterfaceActionType::ToggleButton, FInputChord());
+	UI_COMMAND(MeshRenderingMode_Solid, "Solid", "Solid Mesh Rendering Mode", EUserInterfaceActionType::RadioButton, FInputChord());
+	UI_COMMAND(MeshRenderingMode_Wireframe, "Wireframe", "Wireframe Mesh Rendering Mode", EUserInterfaceActionType::RadioButton, FInputChord());
+	UI_COMMAND(MeshRenderingMode_None, "None", "No Mesh Rendering Mode", EUserInterfaceActionType::RadioButton, FInputChord());
+	UI_COMMAND(CollisionRenderingMode_Solid, "Solid", "Solid Collision Rendering Mode", EUserInterfaceActionType::RadioButton, FInputChord());
+	UI_COMMAND(CollisionRenderingMode_Wireframe, "Wireframe", "Wireframe Collision Rendering Mode", EUserInterfaceActionType::RadioButton, FInputChord());
+	UI_COMMAND(CollisionRenderingMode_None, "None", "No Collision Rendering Mode", EUserInterfaceActionType::RadioButton, FInputChord());
+	UI_COMMAND(ConstraintRenderingMode_None, "None", "No Constraint Rendering Mode", EUserInterfaceActionType::RadioButton, FInputChord());
+	UI_COMMAND(ConstraintRenderingMode_AllPositions, "All Positions", "All Positions Constraint Rendering Mode", EUserInterfaceActionType::RadioButton, FInputChord());
+	UI_COMMAND(ConstraintRenderingMode_AllLimits, "All Limits", "All Limits Constraint Rendering Mode", EUserInterfaceActionType::RadioButton, FInputChord());
+	UI_COMMAND(ShowKinematicBodies, "Kinematic Bodies", "Displays Kinematic Bodies In Red", EUserInterfaceActionType::ToggleButton, FInputChord());
+	UI_COMMAND(DrawGroundBox, "Ground Box", "Displays Floor Grid", EUserInterfaceActionType::ToggleButton, FInputChord());
+	UI_COMMAND(ToggleGraphicsHierarchy, "Hierarchy", "Show Graphical Hierarchy Of Joints In Preview Viewport", EUserInterfaceActionType::ToggleButton, FInputChord());
+	UI_COMMAND(ToggleBoneInfuences, "Bone Influences", "Displays Vertices Weighted To Currently Selected Bone or Body", EUserInterfaceActionType::ToggleButton, FInputChord());
+	UI_COMMAND(ToggleMassProperties, "Mass Properties", "Show Mass Properties For Bodies When Simulation Is Enabled", EUserInterfaceActionType::ToggleButton, FInputChord());
+	UI_COMMAND(DisableCollision, "Collision Off", "Disable Collision Of Currently Selected Body With Next Selected Body", EUserInterfaceActionType::Button, FInputChord(EKeys::RightBracket));
+	UI_COMMAND(EnableCollision, "Collision On", "Enable Collision Of Currently Selected Body With Next Selected Body", EUserInterfaceActionType::Button, FInputChord(EKeys::LeftBracket));
+	UI_COMMAND(WeldToBody, "Weld", "Weld Body: Weld Currently Selected Bodies", EUserInterfaceActionType::Button, FInputChord());
+	UI_COMMAND(AddNewBody, "New Body", "Add New Body To Selected Bone.", EUserInterfaceActionType::Button, FInputChord());
+	UI_COMMAND(AddSphere, "Add Sphere", "Add Sphere To Selected Bone", EUserInterfaceActionType::Button, FInputChord());
+	UI_COMMAND(AddSphyl, "Add Sphyl", "Add Sphyl To Selected Bone", EUserInterfaceActionType::Button, FInputChord());
+	UI_COMMAND(AddBox, "Add Box", "Add Box To Selected Bone", EUserInterfaceActionType::Button, FInputChord());
+	UI_COMMAND(DeletePrimitive, "Delete", "Delete Selected Primitive", EUserInterfaceActionType::Button, FInputChord());
+	UI_COMMAND(DuplicatePrimitive, "Duplicate", "Duplicate Selected Primitive", EUserInterfaceActionType::Button, FInputChord());
+	UI_COMMAND(ResetConstraint, "Reset", "Reset Constraint", EUserInterfaceActionType::Button, FInputChord());
+	UI_COMMAND(SnapConstraint, "Snap", "Snap Constraint To Bone", EUserInterfaceActionType::Button, FInputChord());
+	UI_COMMAND(ConvertToBallAndSocket, "To B&S", "Convert Selected Constraint To Ball-And-Socket", EUserInterfaceActionType::Button, FInputChord());
+	UI_COMMAND(ConvertToHinge, "To Hinge", "Convert Selected Constraint To Hinge", EUserInterfaceActionType::Button, FInputChord());
+	UI_COMMAND(ConvertToPrismatic, "To Prismatic", "Convert Selected Constraint To Prismatic", EUserInterfaceActionType::Button, FInputChord());
+	UI_COMMAND(ConvertToSkeletal, "To Skeletal", "Convert Selected Constraint To Skeletal", EUserInterfaceActionType::Button, FInputChord());
+	UI_COMMAND(DeleteConstraint, "Delete", "Delete Selected Constraint", EUserInterfaceActionType::Button, FInputChord());
+	UI_COMMAND(PlayAnimation, "Play", "Play Animation", EUserInterfaceActionType::ToggleButton, FInputChord());
+	UI_COMMAND(RecordAnimation, "Record", "Record Animation", EUserInterfaceActionType::ToggleButton, FInputChord());
+	UI_COMMAND(ShowSkeleton, "Skeleton", "Show Skeleton", EUserInterfaceActionType::ToggleButton, FInputChord());
+	UI_COMMAND(MakeBodyKinematic, "Kinematic", "Make Body Kinematic", EUserInterfaceActionType::RadioButton, FInputChord());
+	UI_COMMAND(MakeBodySimulated, "Simulated", "Make Body Simulated", EUserInterfaceActionType::RadioButton, FInputChord());
+	UI_COMMAND(MakeBodyDefault, "Default", "Reset This Body To Default", EUserInterfaceActionType::RadioButton, FInputChord());
+	UI_COMMAND(KinematicAllBodiesBelow, "Set All Bodies Below To Kinematic", "Set All Bodies Below To Kinematic", EUserInterfaceActionType::Button, FInputChord());
+	UI_COMMAND(SimulatedAllBodiesBelow, "Set All Bodies Below To Simulated", "Set All Bodies Below To Simulated", EUserInterfaceActionType::Button, FInputChord());
+	UI_COMMAND(MakeAllBodiesBelowDefault, "Reset All Bodies Below To Default", "Reset All Bodies Below To Default", EUserInterfaceActionType::Button, FInputChord());
+	UI_COMMAND(DeleteBody, "Delete", "Delete Selected Body", EUserInterfaceActionType::Button, FInputChord());
+	UI_COMMAND(DeleteAllBodiesBelow, "Delete All Bodies Below", "Delete All Bodies Below", EUserInterfaceActionType::Button, FInputChord());
+	UI_COMMAND(ToggleMotor, "Toggle Motor", "Toggle Motor", EUserInterfaceActionType::Button, FInputChord());
+	UI_COMMAND(EnableMotorsBelow, "Enable Motors Below", "Enable Motors Below", EUserInterfaceActionType::Button, FInputChord());
+	UI_COMMAND(DisableMotorsBelow, "Disable Motors Below", "Disable Motors Below", EUserInterfaceActionType::Button, FInputChord());
+	UI_COMMAND(SelectAllObjects, "Select All Objects", "Select All Objects", EUserInterfaceActionType::Button, FInputChord(EModifierKey::Control, EKeys::A));
+	UI_COMMAND(HierarchyFilterAll, "All Bones", "Show Entire Hierarchy", EUserInterfaceActionType::Button, FInputChord());
+	UI_COMMAND(HierarchyFilterBodies, "Bones With Bodies", "Filter Bones With Bodies", EUserInterfaceActionType::Button, FInputChord());
+	UI_COMMAND(PerspectiveView, "Perspective", "Perspective", EUserInterfaceActionType::Button, FInputChord());
+	UI_COMMAND(LeftView, "Left", "Orthographic view from left", EUserInterfaceActionType::Button, FInputChord());
+	UI_COMMAND(TopView, "Top", "Orthographic view from top", EUserInterfaceActionType::Button, FInputChord());
+	UI_COMMAND(FrontView, "Front", "Orthographic view from front", EUserInterfaceActionType::Button, FInputChord());
+	UI_COMMAND(RightView, "Right", "Orthographic view from right", EUserInterfaceActionType::Button, FInputChord());
+	UI_COMMAND(BottomView, "Bottom", "Orthographic view from bottom", EUserInterfaceActionType::Button, FInputChord());
+	UI_COMMAND(BackView, "Back", "Orthographic view from back", EUserInterfaceActionType::Button, FInputChord());
+	UI_COMMAND(Mirror, "Mirror", "Finds the body on the other side and duplicates constraint and body", EUserInterfaceActionType::Button, FInputChord(EKeys::M));
 
 
-	UI_COMMAND(SelectionLock, "Lock Selection", "", EUserInterfaceActionType::Button, FInputGesture(EKeys::X));
+	UI_COMMAND(SelectionLock, "Lock Selection", "", EUserInterfaceActionType::Button, FInputChord(EKeys::X));
 
 	/** As two commands cannot have the same key; this command wraps both 
 	  * DeletePrimitive and DeleteConstraint so the user can delete whatever is selected 
 	  */
-	UI_COMMAND(DeleteSelected, "Delete selected primitive or constraint", "", EUserInterfaceActionType::Button, FInputGesture(EKeys::Platform_Delete));
-	UI_COMMAND(CycleConstraintOrientation, "Cycle selected constraint orientation", "", EUserInterfaceActionType::Button, FInputGesture(EKeys::Q));
-	UI_COMMAND(CycleConstraintActive, "Cycle active constraint", "", EUserInterfaceActionType::Button, FInputGesture(EKeys::Four));
-	UI_COMMAND(ToggleSwing1, "Toggle Swing1 Constraint", "", EUserInterfaceActionType::Button, FInputGesture(EKeys::One));
-	UI_COMMAND(ToggleSwing2, "Toggle Swing2 Constraint", "", EUserInterfaceActionType::Button, FInputGesture(EKeys::Three));
-	UI_COMMAND(ToggleTwist, "Toggle Twist Constraint", "", EUserInterfaceActionType::Button, FInputGesture(EKeys::Two));
-	UI_COMMAND(FocusOnSelection, "Focus the viewport on the current selection", "", EUserInterfaceActionType::Button, FInputGesture(EKeys::F));
+	UI_COMMAND(DeleteSelected, "Delete selected primitive or constraint", "", EUserInterfaceActionType::Button, FInputChord(EKeys::Platform_Delete));
+	UI_COMMAND(CycleConstraintOrientation, "Cycle selected constraint orientation", "", EUserInterfaceActionType::Button, FInputChord(EKeys::Q));
+	UI_COMMAND(CycleConstraintActive, "Cycle active constraint", "", EUserInterfaceActionType::Button, FInputChord(EKeys::Four));
+	UI_COMMAND(ToggleSwing1, "Toggle Swing1 Constraint", "", EUserInterfaceActionType::Button, FInputChord(EKeys::One));
+	UI_COMMAND(ToggleSwing2, "Toggle Swing2 Constraint", "", EUserInterfaceActionType::Button, FInputChord(EKeys::Three));
+	UI_COMMAND(ToggleTwist, "Toggle Twist Constraint", "", EUserInterfaceActionType::Button, FInputChord(EKeys::Two));
+	UI_COMMAND(FocusOnSelection, "Focus the viewport on the current selection", "", EUserInterfaceActionType::Button, FInputChord(EKeys::F));
 }
 
 #undef LOCTEXT_NAMESPACE
\ No newline at end of file
diff --git a/Engine/Source/Editor/ReferenceViewer/Private/ReferenceViewerActions.cpp b/Engine/Source/Editor/ReferenceViewer/Private/ReferenceViewerActions.cpp
--- a/Engine/Source/Editor/ReferenceViewer/Private/ReferenceViewerActions.cpp
+++ b/Engine/Source/Editor/ReferenceViewer/Private/ReferenceViewerActions.cpp
@@ -10,15 +10,15 @@
 
 void FReferenceViewerActions::RegisterCommands()
 {
-	UI_COMMAND(OpenSelectedInAssetEditor, "Edit Asset", "Opens the selected asset in the relevent editor.", EUserInterfaceActionType::Button, FInputGesture());
-	UI_COMMAND(ReCenterGraph, "Re-Center Graph", "Re-centers the graph on this node, showing all referencers and dependencies for this asset instead", EUserInterfaceActionType::Button, FInputGesture());
-	UI_COMMAND(ListReferencedObjects, "List Referenced Objects", "Shows a list of objects that the selected asset references.", EUserInterfaceActionType::Button, FInputGesture());
-	UI_COMMAND(ListObjectsThatReference, "List Objects That Reference", "Lists objects that reference the selected asset.", EUserInterfaceActionType::Button, FInputGesture());
-	UI_COMMAND(ShowReferenceTree, "Show Reference Tree...", "Shows a reference tree for the selected asset.", EUserInterfaceActionType::Button, FInputGesture());
+	UI_COMMAND(OpenSelectedInAssetEditor, "Edit Asset", "Opens the selected asset in the relevent editor.", EUserInterfaceActionType::Button, FInputChord());
+	UI_COMMAND(ReCenterGraph, "Re-Center Graph", "Re-centers the graph on this node, showing all referencers and dependencies for this asset instead", EUserInterfaceActionType::Button, FInputChord());
+	UI_COMMAND(ListReferencedObjects, "List Referenced Objects", "Shows a list of objects that the selected asset references.", EUserInterfaceActionType::Button, FInputChord());
+	UI_COMMAND(ListObjectsThatReference, "List Objects That Reference", "Lists objects that reference the selected asset.", EUserInterfaceActionType::Button, FInputChord());
+	UI_COMMAND(ShowReferenceTree, "Show Reference Tree...", "Shows a reference tree for the selected asset.", EUserInterfaceActionType::Button, FInputChord());
 
-	UI_COMMAND(MakeLocalCollectionWithReferencedAssets, "Local", "Local. This collection is only visible to you and is not in source control.", EUserInterfaceActionType::Button, FInputGesture());
-	UI_COMMAND(MakePrivateCollectionWithReferencedAssets, "Private", "Private. This collection is only visible to you.", EUserInterfaceActionType::Button, FInputGesture());
-	UI_COMMAND(MakeSharedCollectionWithReferencedAssets, "Shared", "Shared. This collection is visible to everyone.", EUserInterfaceActionType::Button, FInputGesture());
+	UI_COMMAND(MakeLocalCollectionWithReferencedAssets, "Local", "Local. This collection is only visible to you and is not in source control.", EUserInterfaceActionType::Button, FInputChord());
+	UI_COMMAND(MakePrivateCollectionWithReferencedAssets, "Private", "Private. This collection is only visible to you.", EUserInterfaceActionType::Button, FInputChord());
+	UI_COMMAND(MakeSharedCollectionWithReferencedAssets, "Shared", "Shared. This collection is visible to everyone.", EUserInterfaceActionType::Button, FInputChord());
 }
 
 #undef LOCTEXT_NAMESPACE
\ No newline at end of file
diff --git a/Engine/Source/Editor/Sequencer/Private/SequencerCommands.cpp b/Engine/Source/Editor/Sequencer/Private/SequencerCommands.cpp
--- a/Engine/Source/Editor/Sequencer/Private/SequencerCommands.cpp
+++ b/Engine/Source/Editor/Sequencer/Private/SequencerCommands.cpp
@@ -7,29 +7,29 @@
 
 void FSequencerCommands::RegisterCommands()
 {
-	UI_COMMAND( SetKey, "Set Key", "Sets a key at the current time for the selected actor.", EUserInterfaceActionType::Button, FInputGesture(EKeys::K) );
+	UI_COMMAND( SetKey, "Set Key", "Sets a key at the current time for the selected actor.", EUserInterfaceActionType::Button, FInputChord(EKeys::K) );
 
-	UI_COMMAND( ToggleAutoKeyEnabled, "Auto Key", "Enables and disables auto keying.", EUserInterfaceActionType::ToggleButton, FInputGesture() );
+	UI_COMMAND( ToggleAutoKeyEnabled, "Auto Key", "Enables and disables auto keying.", EUserInterfaceActionType::ToggleButton, FInputChord() );
 
-	UI_COMMAND( ToggleIsSnapEnabled, "Enable Snapping", "Enables and disables snapping.", EUserInterfaceActionType::ToggleButton, FInputGesture() );
+	UI_COMMAND( ToggleIsSnapEnabled, "Enable Snapping", "Enables and disables snapping.", EUserInterfaceActionType::ToggleButton, FInputChord() );
 
-	UI_COMMAND( ToggleSnapKeyTimesToInterval, "Snap to the interval", "Snap keys to the interval.", EUserInterfaceActionType::ToggleButton, FInputGesture() );
-	UI_COMMAND( ToggleSnapKeyTimesToKeys, "Snap to other keys", "Snap keys to other keys in the section.", EUserInterfaceActionType::ToggleButton, FInputGesture() );
+	UI_COMMAND( ToggleSnapKeyTimesToInterval, "Snap to the interval", "Snap keys to the interval.", EUserInterfaceActionType::ToggleButton, FInputChord() );
+	UI_COMMAND( ToggleSnapKeyTimesToKeys, "Snap to other keys", "Snap keys to other keys in the section.", EUserInterfaceActionType::ToggleButton, FInputChord() );
 
-	UI_COMMAND( ToggleSnapSectionTimesToInterval, "Snap to the interval", "Snap sections to the interval.", EUserInterfaceActionType::ToggleButton, FInputGesture() );
-	UI_COMMAND( ToggleSnapSectionTimesToSections, "Snap to other sections", "Snap sections to other sections.", EUserInterfaceActionType::ToggleButton, FInputGesture() );
+	UI_COMMAND( ToggleSnapSectionTimesToInterval, "Snap to the interval", "Snap sections to the interval.", EUserInterfaceActionType::ToggleButton, FInputChord() );
+	UI_COMMAND( ToggleSnapSectionTimesToSections, "Snap to other sections", "Snap sections to other sections.", EUserInterfaceActionType::ToggleButton, FInputChord() );
 
-	UI_COMMAND( ToggleSnapPlayTimeToInterval, "Snap to the interval while scrubbing", "Snap the play time to the interval while scrubbing.", EUserInterfaceActionType::ToggleButton, FInputGesture() );
+	UI_COMMAND( ToggleSnapPlayTimeToInterval, "Snap to the interval while scrubbing", "Snap the play time to the interval while scrubbing.", EUserInterfaceActionType::ToggleButton, FInputChord() );
 
-	UI_COMMAND(ToggleSnapCurveValueToInterval, "Snap curve key values", "Snap curve keys to the value snapping interval.", EUserInterfaceActionType::ToggleButton, FInputGesture() );
+	UI_COMMAND(ToggleSnapCurveValueToInterval, "Snap curve key values", "Snap curve keys to the value snapping interval.", EUserInterfaceActionType::ToggleButton, FInputChord() );
 
-	UI_COMMAND( ToggleCleanView, "Clean View", "Enable 'Clean View' mode which displays only global tracks when no filter is applied.", EUserInterfaceActionType::ToggleButton, FInputGesture() );
+	UI_COMMAND( ToggleCleanView, "Clean View", "Enable 'Clean View' mode which displays only global tracks when no filter is applied.", EUserInterfaceActionType::ToggleButton, FInputChord() );
 	
-	UI_COMMAND( ToggleShowCurveEditor, "Curve Editor", "Show the animation keys in a curve editor.", EUserInterfaceActionType::ToggleButton, FInputGesture() );
-	UI_COMMAND( ToggleShowCurveEditorCurveToolTips, "Curve Tool Tips", "Show a tool tip with name and values when hoving over a curve.", EUserInterfaceActionType::ToggleButton, FInputGesture() );
-	UI_COMMAND( SetAllCurveVisibility, "All Curves", "Show all curves in the curve editor.", EUserInterfaceActionType::RadioButton, FInputGesture() );
-	UI_COMMAND( SetSelectedCurveVisibility, "Selected Curves", "Show only selected curves in the curve editor.", EUserInterfaceActionType::RadioButton, FInputGesture() );
-	UI_COMMAND( SetAnimatedCurveVisibility, "Animated Curves", "Show only curves with keyframes in the curve editor.", EUserInterfaceActionType::RadioButton, FInputGesture() );
+	UI_COMMAND( ToggleShowCurveEditor, "Curve Editor", "Show the animation keys in a curve editor.", EUserInterfaceActionType::ToggleButton, FInputChord() );
+	UI_COMMAND( ToggleShowCurveEditorCurveToolTips, "Curve Tool Tips", "Show a tool tip with name and values when hoving over a curve.", EUserInterfaceActionType::ToggleButton, FInputChord() );
+	UI_COMMAND( SetAllCurveVisibility, "All Curves", "Show all curves in the curve editor.", EUserInterfaceActionType::RadioButton, FInputChord() );
+	UI_COMMAND( SetSelectedCurveVisibility, "Selected Curves", "Show only selected curves in the curve editor.", EUserInterfaceActionType::RadioButton, FInputChord() );
+	UI_COMMAND( SetAnimatedCurveVisibility, "Animated Curves", "Show only curves with keyframes in the curve editor.", EUserInterfaceActionType::RadioButton, FInputChord() );
 }
 
 #undef LOCTEXT_NAMESPACE
\ No newline at end of file
diff --git a/Engine/Source/Editor/StaticMeshEditor/Private/StaticMeshEditorActions.cpp b/Engine/Source/Editor/StaticMeshEditor/Private/StaticMeshEditorActions.cpp
--- a/Engine/Source/Editor/StaticMeshEditor/Private/StaticMeshEditorActions.cpp
+++ b/Engine/Source/Editor/StaticMeshEditor/Private/StaticMeshEditorActions.cpp
@@ -9,42 +9,42 @@
 
 void FStaticMeshEditorCommands::RegisterCommands()
 {
-	UI_COMMAND( SetShowWireframe, "Wireframe", "Toggles the viewmode of the Preview Pane between a lit view and a wireframe view.", EUserInterfaceActionType::ToggleButton, FInputGesture() );
-	UI_COMMAND( SetShowVertexColor, "Vert Colors", "Toggles vertex colors.", EUserInterfaceActionType::ToggleButton, FInputGesture() );
-	UI_COMMAND( SetDrawUVs, "UV", "Toggles display of the static mesh's UVs for the specified channel.", EUserInterfaceActionType::ToggleButton, FInputGesture() );
-	UI_COMMAND( SetShowGrid, "Grid", "Displays the viewport grid.", EUserInterfaceActionType::ToggleButton, FInputGesture() );
-	UI_COMMAND( SetShowBounds, "Bounds", "Toggles display of the bounds of the static mesh.", EUserInterfaceActionType::ToggleButton, FInputGesture() );
-	UI_COMMAND( SetShowCollision, "Collision", "Toggles display of the simplified collision mesh of the static mesh, if one has been assigned.", EUserInterfaceActionType::ToggleButton, FInputGesture() );
-	UI_COMMAND( ResetCamera, "Reset Camera", "Resets the camera to focus on the static mesh", EUserInterfaceActionType::Button, FInputGesture() );
-	UI_COMMAND( SetShowSockets, "Sockets", "Displays the static mesh sockets.", EUserInterfaceActionType::ToggleButton, FInputGesture() );
-	UI_COMMAND( SetDrawAdditionalData, "Additional Data", "Draw additional user data associated with asset.", EUserInterfaceActionType::ToggleButton, FInputGesture() );
-
-	UI_COMMAND( SetShowNormals, "Normals", "Toggles display of vertex normals in the Preview Pane.", EUserInterfaceActionType::ToggleButton, FInputGesture() );
-	UI_COMMAND( SetShowTangents, "Tangents", "Toggles display of vertex tangents in the Preview Pane.", EUserInterfaceActionType::ToggleButton, FInputGesture() );
-	UI_COMMAND( SetShowBinormals, "Binormals", "Toggles display of vertex binormals (orthogonal vector to normal and tangent) in the Preview Pane.", EUserInterfaceActionType::ToggleButton, FInputGesture() );
-	UI_COMMAND( SetShowVertices, "Vertices", "Toggles display of vertices in the Preview Pane.", EUserInterfaceActionType::ToggleButton, FInputGesture() ); 
-	UI_COMMAND( SetShowPivot, "Show Pivot", "Display the pivot location of the static mesh.", EUserInterfaceActionType::ToggleButton, FInputGesture() );
-
-	UI_COMMAND( CreateDOP10X, "Add 10DOP-X Simplified Collision", "Generates a new axis-aligned box collision mesh with the 4 X-axis aligned edges beveled (10 total sides) encompassing the static mesh.", EUserInterfaceActionType::Button, FInputGesture() );
-	UI_COMMAND( CreateDOP10Y, "Add 10DOP-Y Simplified Collision", "Generates a new axis-aligned box collision mesh with the 4 Y-axis aligned edges beveled (10 total sides) encompassing the static mesh.", EUserInterfaceActionType::Button, FInputGesture() );
-	UI_COMMAND( CreateDOP10Z, "Add 10DOP-Z Simplified Collision", "Generates a new axis-aligned box collision mesh with the 4 Z-axis aligned edges beveled (10 total sides) encompassing the static mesh.", EUserInterfaceActionType::Button, FInputGesture() );
-	UI_COMMAND( CreateDOP18, "Add 18DOP Simplified Collision", "Generates a new axis-aligned box collision mesh with all edges beveled (18 total sides) encompassing the static mesh.", EUserInterfaceActionType::Button, FInputGesture() );
-	UI_COMMAND( CreateDOP26, "Add 26DOP Simplified Collision", "Generates a new axis-aligned box collision mesh with all edges and corners beveled (26 total sides) encompassing the static mesh.", EUserInterfaceActionType::Button, FInputGesture() );
-
-	UI_COMMAND( CreateBoxCollision, "Add Box Simplified Collision", "Generates a new box collision mesh encompassing the static mesh.", EUserInterfaceActionType::Button, FInputGesture() );
-	UI_COMMAND( CreateSphereCollision, "Add Sphere Simplified Collision", "Generates a new sphere collision mesh encompassing the static mesh.", EUserInterfaceActionType::Button, FInputGesture() );
-	UI_COMMAND( CreateSphylCollision, "Add Capsule Simplified Collision", "Generates a new capsule collision mesh encompassing the static mesh.", EUserInterfaceActionType::Button, FInputGesture() );
-	UI_COMMAND( CreateAutoConvexCollision, "Auto Convex Collision", "Opens the Auto Convex Collision Tool for generating a new convex collision mesh, or meshes.", EUserInterfaceActionType::Button, FInputGesture() );
-	UI_COMMAND( RemoveCollision, "Remove Collision", "Removes any simplified collision assigned to the static mesh.", EUserInterfaceActionType::Button, FInputGesture() );
-	UI_COMMAND( ConvertBoxesToConvex, "Convert Boxes to Convex", "Converts any simple box collision meshes to convex collision meshes.", EUserInterfaceActionType::Button, FInputGesture() );
-
-	UI_COMMAND( CopyCollisionFromSelectedMesh, "Copy Collision from Selected Static Mesh", "Copy collision from the static mesh selected in Content Browser.", EUserInterfaceActionType::Button, FInputGesture() );
-
-	UI_COMMAND( FindSource, "Find Source", "Opens explorer at the location of this asset.", EUserInterfaceActionType::Button, FInputGesture() );
+	UI_COMMAND( SetShowWireframe, "Wireframe", "Toggles the viewmode of the Preview Pane between a lit view and a wireframe view.", EUserInterfaceActionType::ToggleButton, FInputChord() );
+	UI_COMMAND( SetShowVertexColor, "Vert Colors", "Toggles vertex colors.", EUserInterfaceActionType::ToggleButton, FInputChord() );
+	UI_COMMAND( SetDrawUVs, "UV", "Toggles display of the static mesh's UVs for the specified channel.", EUserInterfaceActionType::ToggleButton, FInputChord() );
+	UI_COMMAND( SetShowGrid, "Grid", "Displays the viewport grid.", EUserInterfaceActionType::ToggleButton, FInputChord() );
+	UI_COMMAND( SetShowBounds, "Bounds", "Toggles display of the bounds of the static mesh.", EUserInterfaceActionType::ToggleButton, FInputChord() );
+	UI_COMMAND( SetShowCollision, "Collision", "Toggles display of the simplified collision mesh of the static mesh, if one has been assigned.", EUserInterfaceActionType::ToggleButton, FInputChord() );
+	UI_COMMAND( ResetCamera, "Reset Camera", "Resets the camera to focus on the static mesh", EUserInterfaceActionType::Button, FInputChord() );
+	UI_COMMAND( SetShowSockets, "Sockets", "Displays the static mesh sockets.", EUserInterfaceActionType::ToggleButton, FInputChord() );
+	UI_COMMAND( SetDrawAdditionalData, "Additional Data", "Draw additional user data associated with asset.", EUserInterfaceActionType::ToggleButton, FInputChord() );
+
+	UI_COMMAND( SetShowNormals, "Normals", "Toggles display of vertex normals in the Preview Pane.", EUserInterfaceActionType::ToggleButton, FInputChord() );
+	UI_COMMAND( SetShowTangents, "Tangents", "Toggles display of vertex tangents in the Preview Pane.", EUserInterfaceActionType::ToggleButton, FInputChord() );
+	UI_COMMAND( SetShowBinormals, "Binormals", "Toggles display of vertex binormals (orthogonal vector to normal and tangent) in the Preview Pane.", EUserInterfaceActionType::ToggleButton, FInputChord() );
+	UI_COMMAND( SetShowVertices, "Vertices", "Toggles display of vertices in the Preview Pane.", EUserInterfaceActionType::ToggleButton, FInputChord() ); 
+	UI_COMMAND( SetShowPivot, "Show Pivot", "Display the pivot location of the static mesh.", EUserInterfaceActionType::ToggleButton, FInputChord() );
+
+	UI_COMMAND( CreateDOP10X, "Add 10DOP-X Simplified Collision", "Generates a new axis-aligned box collision mesh with the 4 X-axis aligned edges beveled (10 total sides) encompassing the static mesh.", EUserInterfaceActionType::Button, FInputChord() );
+	UI_COMMAND( CreateDOP10Y, "Add 10DOP-Y Simplified Collision", "Generates a new axis-aligned box collision mesh with the 4 Y-axis aligned edges beveled (10 total sides) encompassing the static mesh.", EUserInterfaceActionType::Button, FInputChord() );
+	UI_COMMAND( CreateDOP10Z, "Add 10DOP-Z Simplified Collision", "Generates a new axis-aligned box collision mesh with the 4 Z-axis aligned edges beveled (10 total sides) encompassing the static mesh.", EUserInterfaceActionType::Button, FInputChord() );
+	UI_COMMAND( CreateDOP18, "Add 18DOP Simplified Collision", "Generates a new axis-aligned box collision mesh with all edges beveled (18 total sides) encompassing the static mesh.", EUserInterfaceActionType::Button, FInputChord() );
+	UI_COMMAND( CreateDOP26, "Add 26DOP Simplified Collision", "Generates a new axis-aligned box collision mesh with all edges and corners beveled (26 total sides) encompassing the static mesh.", EUserInterfaceActionType::Button, FInputChord() );
+
+	UI_COMMAND( CreateBoxCollision, "Add Box Simplified Collision", "Generates a new box collision mesh encompassing the static mesh.", EUserInterfaceActionType::Button, FInputChord() );
+	UI_COMMAND( CreateSphereCollision, "Add Sphere Simplified Collision", "Generates a new sphere collision mesh encompassing the static mesh.", EUserInterfaceActionType::Button, FInputChord() );
+	UI_COMMAND( CreateSphylCollision, "Add Capsule Simplified Collision", "Generates a new capsule collision mesh encompassing the static mesh.", EUserInterfaceActionType::Button, FInputChord() );
+	UI_COMMAND( CreateAutoConvexCollision, "Auto Convex Collision", "Opens the Auto Convex Collision Tool for generating a new convex collision mesh, or meshes.", EUserInterfaceActionType::Button, FInputChord() );
+	UI_COMMAND( RemoveCollision, "Remove Collision", "Removes any simplified collision assigned to the static mesh.", EUserInterfaceActionType::Button, FInputChord() );
+	UI_COMMAND( ConvertBoxesToConvex, "Convert Boxes to Convex", "Converts any simple box collision meshes to convex collision meshes.", EUserInterfaceActionType::Button, FInputChord() );
+
+	UI_COMMAND( CopyCollisionFromSelectedMesh, "Copy Collision from Selected Static Mesh", "Copy collision from the static mesh selected in Content Browser.", EUserInterfaceActionType::Button, FInputChord() );
+
+	UI_COMMAND( FindSource, "Find Source", "Opens explorer at the location of this asset.", EUserInterfaceActionType::Button, FInputChord() );
 	
-	UI_COMMAND( ChangeMesh, "Change Mesh", "Changes the static mesh asset loaded in the Static Mesh Editor to the asset currently selected in the Content Browser.", EUserInterfaceActionType::Button, FInputGesture() );
+	UI_COMMAND( ChangeMesh, "Change Mesh", "Changes the static mesh asset loaded in the Static Mesh Editor to the asset currently selected in the Content Browser.", EUserInterfaceActionType::Button, FInputChord() );
 
-	UI_COMMAND( SaveGeneratedLODs, "Save Generated LODs", "Saves the generated LODs in the mesh package.", EUserInterfaceActionType::Button, FInputGesture());
+	UI_COMMAND( SaveGeneratedLODs, "Save Generated LODs", "Saves the generated LODs in the mesh package.", EUserInterfaceActionType::Button, FInputChord());
 }
 
 #undef LOCTEXT_NAMESPACE
\ No newline at end of file
diff --git a/Engine/Source/Editor/TextureEditor/Private/Models/TextureEditorCommands.cpp b/Engine/Source/Editor/TextureEditor/Private/Models/TextureEditorCommands.cpp
--- a/Engine/Source/Editor/TextureEditor/Private/Models/TextureEditorCommands.cpp
+++ b/Engine/Source/Editor/TextureEditor/Private/Models/TextureEditorCommands.cpp
@@ -6,21 +6,21 @@
 
 void FTextureEditorCommands::RegisterCommands()
 {
-	UI_COMMAND(RedChannel, "Red", "Toggles the red channel", EUserInterfaceActionType::ToggleButton, FInputGesture());
-	UI_COMMAND(GreenChannel, "Green", "Toggles the green channel", EUserInterfaceActionType::ToggleButton, FInputGesture());
-	UI_COMMAND(BlueChannel, "Blue", "Toggles the blue channel", EUserInterfaceActionType::ToggleButton, FInputGesture());
-	UI_COMMAND(AlphaChannel, "Alpha", "Toggles the alpha channel", EUserInterfaceActionType::ToggleButton, FInputGesture());
-	UI_COMMAND(Desaturation, "Desaturation", "Toggles color desaturation", EUserInterfaceActionType::ToggleButton, FInputGesture());
+	UI_COMMAND(RedChannel, "Red", "Toggles the red channel", EUserInterfaceActionType::ToggleButton, FInputChord());
+	UI_COMMAND(GreenChannel, "Green", "Toggles the green channel", EUserInterfaceActionType::ToggleButton, FInputChord());
+	UI_COMMAND(BlueChannel, "Blue", "Toggles the blue channel", EUserInterfaceActionType::ToggleButton, FInputChord());
+	UI_COMMAND(AlphaChannel, "Alpha", "Toggles the alpha channel", EUserInterfaceActionType::ToggleButton, FInputChord());
+	UI_COMMAND(Desaturation, "Desaturation", "Toggles color desaturation", EUserInterfaceActionType::ToggleButton, FInputChord());
 	
-	UI_COMMAND(CheckeredBackground, "Checkered", "Checkered background pattern behind the texture", EUserInterfaceActionType::RadioButton, FInputGesture());
-	UI_COMMAND(CheckeredBackgroundFill, "Checkered (Fill)", "Checkered background pattern behind the entire viewport", EUserInterfaceActionType::RadioButton, FInputGesture());
-	UI_COMMAND(FitToViewport, "Scale To Fit", "If enabled, the texture will be scaled to fit the viewport", EUserInterfaceActionType::ToggleButton, FInputGesture());
-	UI_COMMAND(SolidBackground, "Solid Color", "Solid color background", EUserInterfaceActionType::RadioButton, FInputGesture());
-	UI_COMMAND(TextureBorder, "Draw Border", "If enabled, a border is drawn around the texture", EUserInterfaceActionType::ToggleButton, FInputGesture());
+	UI_COMMAND(CheckeredBackground, "Checkered", "Checkered background pattern behind the texture", EUserInterfaceActionType::RadioButton, FInputChord());
+	UI_COMMAND(CheckeredBackgroundFill, "Checkered (Fill)", "Checkered background pattern behind the entire viewport", EUserInterfaceActionType::RadioButton, FInputChord());
+	UI_COMMAND(FitToViewport, "Scale To Fit", "If enabled, the texture will be scaled to fit the viewport", EUserInterfaceActionType::ToggleButton, FInputChord());
+	UI_COMMAND(SolidBackground, "Solid Color", "Solid color background", EUserInterfaceActionType::RadioButton, FInputChord());
+	UI_COMMAND(TextureBorder, "Draw Border", "If enabled, a border is drawn around the texture", EUserInterfaceActionType::ToggleButton, FInputChord());
 
-	UI_COMMAND(CompressNow, "Compress", "Compress the texture", EUserInterfaceActionType::Button, FInputGesture());
-	UI_COMMAND(Reimport, "Reimport", "Reimports the texture from file", EUserInterfaceActionType::Button, FInputGesture());
-	UI_COMMAND(Settings, "Settings...", "Opens the settings for the texture editor", EUserInterfaceActionType::Button, FInputGesture());
+	UI_COMMAND(CompressNow, "Compress", "Compress the texture", EUserInterfaceActionType::Button, FInputChord());
+	UI_COMMAND(Reimport, "Reimport", "Reimports the texture from file", EUserInterfaceActionType::Button, FInputChord());
+	UI_COMMAND(Settings, "Settings...", "Opens the settings for the texture editor", EUserInterfaceActionType::Button, FInputChord());
 }
 
 #undef LOCTEXT_NAMESPACE
\ No newline at end of file
diff --git a/Engine/Source/Editor/TranslationEditor/Private/TranslationEditorMenu.cpp b/Engine/Source/Editor/TranslationEditor/Private/TranslationEditorMenu.cpp
--- a/Engine/Source/Editor/TranslationEditor/Private/TranslationEditorMenu.cpp
+++ b/Engine/Source/Editor/TranslationEditor/Private/TranslationEditorMenu.cpp
@@ -94,14 +94,14 @@ void FTranslationEditorMenu::SetupTranslationEditorToolbar( TSharedPtr< FExtende
 
 void FTranslationEditorCommands::RegisterCommands() 
 {
-	UI_COMMAND( ChangeSourceFont, "Change Source Font", "Change the Font for the Source Lanugage", EUserInterfaceActionType::Button, FInputGesture() );
-	UI_COMMAND( ChangeTranslationTargetFont, "Change Translation Font", "Change the Translation Target Language Font", EUserInterfaceActionType::Button, FInputGesture() );
-	UI_COMMAND( SaveTranslations, "Save", "Saves the translations to file", EUserInterfaceActionType::Button, FInputGesture() );
-	UI_COMMAND( PreviewAllTranslationsInEditor, "Preview in Editor", "Preview All Translations in the Editor UI", EUserInterfaceActionType::Button, FInputGesture() );
-	UI_COMMAND( ExportToPortableObjectFormat, "Export to .PO", "Export to Portable Object Format", EUserInterfaceActionType::Button, FInputGesture() );
-	UI_COMMAND( ImportFromPortableObjectFormat, "Import from .PO", "Import from Portable Object Format", EUserInterfaceActionType::Button, FInputGesture() );
-	UI_COMMAND( OpenSearchTab, "Search", "Search Source and Translation Strings", EUserInterfaceActionType::Button, FInputGesture() );
-	UI_COMMAND( OpenTranslationPicker, "Translation Picker", "Open the Translation Picker to Modify Editor Translations", EUserInterfaceActionType::Button, FInputGesture() )
+	UI_COMMAND( ChangeSourceFont, "Change Source Font", "Change the Font for the Source Lanugage", EUserInterfaceActionType::Button, FInputChord() );
+	UI_COMMAND( ChangeTranslationTargetFont, "Change Translation Font", "Change the Translation Target Language Font", EUserInterfaceActionType::Button, FInputChord() );
+	UI_COMMAND( SaveTranslations, "Save", "Saves the translations to file", EUserInterfaceActionType::Button, FInputChord() );
+	UI_COMMAND( PreviewAllTranslationsInEditor, "Preview in Editor", "Preview All Translations in the Editor UI", EUserInterfaceActionType::Button, FInputChord() );
+	UI_COMMAND( ExportToPortableObjectFormat, "Export to .PO", "Export to Portable Object Format", EUserInterfaceActionType::Button, FInputChord() );
+	UI_COMMAND( ImportFromPortableObjectFormat, "Import from .PO", "Import from Portable Object Format", EUserInterfaceActionType::Button, FInputChord() );
+	UI_COMMAND( OpenSearchTab, "Search", "Search Source and Translation Strings", EUserInterfaceActionType::Button, FInputChord() );
+	UI_COMMAND( OpenTranslationPicker, "Translation Picker", "Open the Translation Picker to Modify Editor Translations", EUserInterfaceActionType::Button, FInputChord() )
 }
 
 
diff --git a/Engine/Source/Editor/UMGEditor/Private/Designer/DesignerCommands.cpp b/Engine/Source/Editor/UMGEditor/Private/Designer/DesignerCommands.cpp
--- a/Engine/Source/Editor/UMGEditor/Private/Designer/DesignerCommands.cpp
+++ b/Engine/Source/Editor/UMGEditor/Private/Designer/DesignerCommands.cpp
@@ -7,11 +7,11 @@
 
 void FDesignerCommands::RegisterCommands()
 {
-	UI_COMMAND( LayoutTransform, "Layout Transform Mode", "Adjust widget layout transform", EUserInterfaceActionType::ToggleButton, FInputGesture(EKeys::W) );
-	UI_COMMAND( RenderTransform, "Render Transform Mode", "Adjust widget render transform", EUserInterfaceActionType::ToggleButton, FInputGesture(EKeys::E) );
+	UI_COMMAND( LayoutTransform, "Layout Transform Mode", "Adjust widget layout transform", EUserInterfaceActionType::ToggleButton, FInputChord(EKeys::W) );
+	UI_COMMAND( RenderTransform, "Render Transform Mode", "Adjust widget render transform", EUserInterfaceActionType::ToggleButton, FInputChord(EKeys::E) );
 
-	UI_COMMAND( LocationGridSnap, "Grid Snap", "Enables or disables snapping to the grid when dragging objects around", EUserInterfaceActionType::ToggleButton, FInputGesture() );
-	UI_COMMAND( RotationGridSnap, "Rotation Snap", "Enables or disables snapping objects to a rotation grid", EUserInterfaceActionType::ToggleButton, FInputGesture() );
+	UI_COMMAND( LocationGridSnap, "Grid Snap", "Enables or disables snapping to the grid when dragging objects around", EUserInterfaceActionType::ToggleButton, FInputChord() );
+	UI_COMMAND( RotationGridSnap, "Rotation Snap", "Enables or disables snapping objects to a rotation grid", EUserInterfaceActionType::ToggleButton, FInputChord() );
 }
 
 #undef LOCTEXT_NAMESPACE
\ No newline at end of file
diff --git a/Engine/Source/Editor/UnrealEd/Private/EditorShowFlags.cpp b/Engine/Source/Editor/UnrealEd/Private/EditorShowFlags.cpp
--- a/Engine/Source/Editor/UnrealEd/Private/EditorShowFlags.cpp
+++ b/Engine/Source/Editor/UnrealEd/Private/EditorShowFlags.cpp
@@ -13,11 +13,11 @@ FShowFlagData::FShowFlagData(const FString& InName, const FText& InDisplayName,
 	ShowFlagName = *InName;
 }
 
-FShowFlagData::FShowFlagData(const FString& InName, const FText& InDisplayName, const uint32 InEngineShowFlagIndex, EShowFlagGroup InGroup, const FInputGesture& InInputGesture)
+FShowFlagData::FShowFlagData(const FString& InName, const FText& InDisplayName, const uint32 InEngineShowFlagIndex, EShowFlagGroup InGroup, const FInputChord& InInputChord)
 	:	DisplayName( InDisplayName )
 	,	EngineShowFlagIndex( InEngineShowFlagIndex )
 	,	Group( InGroup )
-	,	InputGesture( InInputGesture )
+	,	InputChord( InInputChord )
 {
 	ShowFlagName = *InName;
 }
@@ -45,21 +45,21 @@ TArray<FShowFlagData>& GetShowFlagMenuItems()
 
 		// FEngineShowFlags 
 		{
-			TMap<FString, FInputGesture> EngineFlagsGestures;
+			TMap<FString, FInputChord> EngineFlagsChords;
 
-			EngineFlagsGestures.Add("Navigation", FInputGesture(EKeys::P) );
-			EngineFlagsGestures.Add("BSP", FInputGesture() );
-			EngineFlagsGestures.Add("Collision", FInputGesture(EKeys::C, EModifierKey::Alt) );
-			EngineFlagsGestures.Add("Fog", FInputGesture(EKeys::F, EModifierKey::Alt ) );
-			EngineFlagsGestures.Add("LightRadius", FInputGesture(EKeys::R, EModifierKey::Alt) );
-			EngineFlagsGestures.Add("StaticMeshes", FInputGesture() );
-			EngineFlagsGestures.Add("Landscape", FInputGesture(EKeys::L, EModifierKey::Alt) );
-			EngineFlagsGestures.Add("Volumes", FInputGesture(EKeys::O, EModifierKey::Alt) );
+			EngineFlagsChords.Add("Navigation", FInputChord(EKeys::P) );
+			EngineFlagsChords.Add("BSP", FInputChord() );
+			EngineFlagsChords.Add("Collision", FInputChord(EKeys::C, EModifierKey::Alt) );
+			EngineFlagsChords.Add("Fog", FInputChord(EKeys::F, EModifierKey::Alt ) );
+			EngineFlagsChords.Add("LightRadius", FInputChord(EKeys::R, EModifierKey::Alt) );
+			EngineFlagsChords.Add("StaticMeshes", FInputChord() );
+			EngineFlagsChords.Add("Landscape", FInputChord(EKeys::L, EModifierKey::Alt) );
+			EngineFlagsChords.Add("Volumes", FInputChord(EKeys::O, EModifierKey::Alt) );
 
 			struct FIterSink
 			{
-				FIterSink(TArray<FShowFlagData>& InShowFlagData, const TMap<FString, FInputGesture>& InGesturesMap) 
-					: ShowFlagData(InShowFlagData), GesturesMap(InGesturesMap)
+				FIterSink(TArray<FShowFlagData>& InShowFlagData, const TMap<FString, FInputChord>& InChordsMap) 
+					: ShowFlagData(InShowFlagData), ChordsMap(InChordsMap)
 				{
 				}
 
@@ -71,10 +71,10 @@ TArray<FShowFlagData>& GetShowFlagMenuItems()
 						FText FlagDisplayName;
 						FEngineShowFlags::FindShowFlagDisplayName(InName, FlagDisplayName);
 
-						const FInputGesture* Gesture = GesturesMap.Find(InName);
-						if (Gesture != NULL)
+						const FInputChord* Chord = ChordsMap.Find(InName);
+						if (Chord != NULL)
 						{
-							ShowFlagData.Add( FShowFlagData( InName, FlagDisplayName, InIndex, Group, *Gesture ) );
+							ShowFlagData.Add( FShowFlagData( InName, FlagDisplayName, InIndex, Group, *Chord ) );
 						}
 						else
 						{
@@ -85,10 +85,10 @@ TArray<FShowFlagData>& GetShowFlagMenuItems()
 				}
 
 				TArray<FShowFlagData>& ShowFlagData;
-				const TMap<FString, FInputGesture>& GesturesMap;
+				const TMap<FString, FInputChord>& ChordsMap;
 			};
 
-			FIterSink Sink(OutShowFlags, EngineFlagsGestures);
+			FIterSink Sink(OutShowFlags, EngineFlagsChords);
 
 			FEngineShowFlags::IterateAllFlags(Sink);
 		}	
