#ANIM: PHAT editor record feature works when non simultaing

https://jira.ol.epicgames.net/browse/UE-7313

diff --git a/Engine/Source/Editor/PhAT/Private/PhATPreviewViewportClient.cpp b/Engine/Source/Editor/PhAT/Private/PhATPreviewViewportClient.cpp
--- a/Engine/Source/Editor/PhAT/Private/PhATPreviewViewportClient.cpp
+++ b/Engine/Source/Editor/PhAT/Private/PhATPreviewViewportClient.cpp
@@ -677,12 +677,12 @@ void FPhATEdPreviewViewportClient::Tick(float DeltaSeconds)
 		// We back up the transforms array now
 		SharedData->EditorSkelComp->AnimationSpaceBases = SharedData->EditorSkelComp->GetSpaceBases();
 		SharedData->EditorSkelComp->SetPhysicsBlendWeight(SharedData->EditorSimOptions->PhysicsBlend);
+	}
 
-		if(SharedData->Recorder.InRecording())
-		{
-			// make sure you don't allow switch SharedData->EditorSkelComp
-			SharedData->Recorder.UpdateRecord(SharedData->EditorSkelComp, DeltaSeconds);
-		}
+	if(SharedData->Recorder.InRecording())
+	{
+		// make sure you don't allow switch SharedData->EditorSkelComp
+		SharedData->Recorder.UpdateRecord(SharedData->EditorSkelComp, DeltaSeconds);
 	}
 }
 
