Expose MaxFPS to PhAT UI.

UE-12540

diff --git a/Engine/Source/Editor/PhAT/Private/PhAT.cpp b/Engine/Source/Editor/PhAT/Private/PhAT.cpp
--- a/Engine/Source/Editor/PhAT/Private/PhAT.cpp
+++ b/Engine/Source/Editor/PhAT/Private/PhAT.cpp
@@ -2326,14 +2326,18 @@ void FPhAT::FixPhysicsState()
 
 void FPhAT::ImpToggleSimulation()
 {
+	static const int32 PrevMaxFPS = GEngine->GetMaxFPS();
+
 	TSharedPtr<FPhATEdPreviewViewportClient> PreviewClient = PreviewViewport->GetViewportClient();
 	if(!SharedData->bRunningSimulation)
 	{
 		BeforeSimulationWidgetMode = PreviewClient->GetWidgetMode();
-		PreviewClient->SetWidgetMode(FWidget::EWidgetMode::WM_None);
+		PreviewClient->SetWidgetMode(FWidget::EWidgetMode::WM_None);	
+		GEngine->SetMaxFPS(SharedData->EditorSimOptions->MaxFPS);
 	}else
 	{
 		PreviewClient->SetWidgetMode(BeforeSimulationWidgetMode);
+		GEngine->SetMaxFPS(PrevMaxFPS);
 	}
 	
 
diff --git a/Engine/Source/Editor/UnrealEd/Classes/Preferences/PhATSimOptions.h b/Engine/Source/Editor/UnrealEd/Classes/Preferences/PhATSimOptions.h
--- a/Engine/Source/Editor/UnrealEd/Classes/Preferences/PhATSimOptions.h
+++ b/Engine/Source/Editor/UnrealEd/Classes/Preferences/PhATSimOptions.h
@@ -28,6 +28,10 @@ class UNREALED_API UPhATSimOptions : public UObject
 	UPROPERTY(EditAnywhere, config, Category=Simulation)
 	float GravScale;
 
+	/** Max FPS for simulation in PhAT. This is helpful for targeting the same FPS as your game. -1 means disabled*/
+	UPROPERTY(EditAnywhere, config, Category = Simulation)
+	int32 MaxFPS;
+
 	/** Linear damping of mouse spring forces */
 	UPROPERTY(EditAnywhere, config, Category=MouseSpring)
 	float HandleLinearDamping;
diff --git a/Engine/Source/Editor/UnrealEd/Private/PreferenceStubs.cpp b/Engine/Source/Editor/UnrealEd/Private/PreferenceStubs.cpp
--- a/Engine/Source/Editor/UnrealEd/Private/PreferenceStubs.cpp
+++ b/Engine/Source/Editor/UnrealEd/Private/PreferenceStubs.cpp
@@ -21,6 +21,7 @@ UPhATSimOptions::UPhATSimOptions(const FObjectInitializer& ObjectInitializer)
 	: Super(ObjectInitializer)
 {
 	PhysicsBlend = 1.0f;
+	MaxFPS = -1;
 }
 
 UMaterialEditorOptions::UMaterialEditorOptions(const FObjectInitializer& ObjectInitializer)
diff --git a/Engine/Source/Runtime/Engine/Classes/Engine/Engine.h b/Engine/Source/Runtime/Engine/Classes/Engine/Engine.h
--- a/Engine/Source/Runtime/Engine/Classes/Engine/Engine.h
+++ b/Engine/Source/Runtime/Engine/Classes/Engine/Engine.h
@@ -1748,6 +1748,12 @@ public:
 	/** Get tick rate limiter. */
 	virtual float GetMaxTickRate(float DeltaTime, bool bAllowFrameRateSmoothing = true) const;
 
+	/** Get max fps. */
+	virtual int32 GetMaxFPS() const;
+
+	/** Set max fps. Overrides console variable. */
+	virtual void SetMaxFPS(const int32 MaxFPS);
+
 	/** Updates the running average delta time */
 	virtual void UpdateRunningAverageDeltaTime(float DeltaTime, bool bAllowFrameRateSmoothing = true);
 
diff --git a/Engine/Source/Runtime/Engine/Private/UnrealEngine.cpp b/Engine/Source/Runtime/Engine/Private/UnrealEngine.cpp
--- a/Engine/Source/Runtime/Engine/Private/UnrealEngine.cpp
+++ b/Engine/Source/Runtime/Engine/Private/UnrealEngine.cpp
@@ -6186,6 +6186,17 @@ float UEngine::GetMaxTickRate(float DeltaTime, bool bAllowFrameRateSmoothing) co
 	return MaxTickRate;
 }
 
+int32 UEngine::GetMaxFPS() const
+{
+	return CVarMaxFPS.GetValueOnAnyThread();
+}
+
+void UEngine::SetMaxFPS(const int32 MaxFPS)
+{
+	IConsoleVariable* ConsoleVariable = CVarMaxFPS.AsVariable();
+	ConsoleVariable->Set(MaxFPS);
+}
+
 /**
  * Enables or disables the ScreenSaver (desktop only)
  *
