Fix Phat simulation not reverting the engine framerate to the correct settings

diff --git a/Engine/Source/Editor/PhAT/Private/PhAT.cpp b/Engine/Source/Editor/PhAT/Private/PhAT.cpp
--- a/Engine/Source/Editor/PhAT/Private/PhAT.cpp
+++ b/Engine/Source/Editor/PhAT/Private/PhAT.cpp
@@ -186,6 +186,12 @@ TSharedRef<SDockTab> FPhAT::SpawnTab( const FSpawnTabArgs& TabSpawnArgs, FName T
 
 FPhAT::~FPhAT()
 {
+	if( SharedData->bRunningSimulation )
+	{
+		// Disable simulation when shutting down
+		ImpToggleSimulation();
+	}
+
 	GEditor->UnregisterForUndo(this);
 }
 
@@ -2334,7 +2340,8 @@ void FPhAT::ImpToggleSimulation()
 		BeforeSimulationWidgetMode = PreviewClient->GetWidgetMode();
 		PreviewClient->SetWidgetMode(FWidget::EWidgetMode::WM_None);	
 		GEngine->SetMaxFPS(SharedData->EditorSimOptions->MaxFPS);
-	}else
+	}
+	else
 	{
 		PreviewClient->SetWidgetMode(BeforeSimulationWidgetMode);
 		GEngine->SetMaxFPS(PrevMaxFPS);
