Fix show kinematic bodies for the case where they blend. Also turn it on by default

diff --git a/Engine/Source/Editor/PhAT/Private/PhATRender.cpp b/Engine/Source/Editor/PhAT/Private/PhATRender.cpp
--- a/Engine/Source/Editor/PhAT/Private/PhATRender.cpp
+++ b/Engine/Source/Editor/PhAT/Private/PhATRender.cpp
@@ -14,7 +14,7 @@ UPhATEdSkeletalMeshComponent::UPhATEdSkeletalMeshComponent(const FObjectInitiali
 	, ElemSelectedColor(255,166,0)
 	, ElemSelectedBodyColor(255, 255, 100)
 	, NoCollisionColor(200, 200, 200)
-	, FixedColor(225,64,64)
+	, FixedColor(125,125,0)
 	, ConstraintBone1Color(255,166,0)
 	, ConstraintBone2Color(0,150,150)
 	, HierarchyDrawColor(220, 255, 220)
@@ -438,19 +438,6 @@ FColor UPhATEdSkeletalMeshComponent::GetPrimitiveColor(int32 BodyIndex, EKCollis
 		}
 	}
 
-	if (SharedData->bRunningSimulation)
-	{
-		// @todo draw fixed, too?
-		if (SharedData->bShowFixedStatus && BodySetup->PhysicsType == PhysType_Simulated)
-		{
-			return FixedColor;
-		}
-		else
-		{
-			return BoneUnselectedColor;
-		}
-	}
-
 	if (SharedData->EditingMode == FPhATSharedData::PEM_ConstraintEdit)
 	{
 		return BoneUnselectedColor;
@@ -479,27 +466,25 @@ FColor UPhATEdSkeletalMeshComponent::GetPrimitiveColor(int32 BodyIndex, EKCollis
 	
 	if (SharedData->bShowFixedStatus)
 	{
-		if (BodySetup->PhysicsType == PhysType_Simulated)
+		const bool bIsSimulatedAtAll = BodySetup->PhysicsType == PhysType_Simulated || (BodySetup->PhysicsType == PhysType_Default && SharedData->EditorSimOptions->PhysicsBlend > 0.f);
+		if (!bIsSimulatedAtAll)
 		{
 			return FixedColor;
 		}
-		else
-		{
-			return BoneUnselectedColor;
-		}
 	}
 	else
 	{
-		// If there is no collision with this body, use 'no collision material'.
-		if (SharedData->NoCollisionBodies.Find(BodyIndex) != INDEX_NONE)
+		if (!SharedData->bRunningSimulation)
 		{
-			return NoCollisionColor;
-		}
-		else
-		{
-			return BoneUnselectedColor;
+			// If there is no collision with this body, use 'no collision material'.
+			if (SharedData->NoCollisionBodies.Find(BodyIndex) != INDEX_NONE)
+			{
+				return NoCollisionColor;
+			}
 		}
 	}
+
+	return BoneUnselectedColor;
 }
 
 UMaterialInterface* UPhATEdSkeletalMeshComponent::GetPrimitiveMaterial(int32 BodyIndex, EKCollisionPrimitiveType PrimitiveType, int32 PrimitiveIndex, bool bHitTest)
diff --git a/Engine/Source/Editor/PhAT/Private/PhATSharedData.cpp b/Engine/Source/Editor/PhAT/Private/PhATSharedData.cpp
--- a/Engine/Source/Editor/PhAT/Private/PhATSharedData.cpp
+++ b/Engine/Source/Editor/PhAT/Private/PhATSharedData.cpp
@@ -43,7 +43,7 @@ FPhATSharedData::FPhATSharedData()
 	bShowHierarchy = false;
 	bShowInfluences = false;
 	bDrawGround = true;
-	bShowFixedStatus = false;
+	bShowFixedStatus = true;
 	bShowAnimSkel = false;
 
 	bSelectionLock = false;
