Fix kinematic bodies changing color when not simulated.

diff --git a/Engine/Source/Editor/PhAT/Private/PhATRender.cpp b/Engine/Source/Editor/PhAT/Private/PhATRender.cpp
--- a/Engine/Source/Editor/PhAT/Private/PhATRender.cpp
+++ b/Engine/Source/Editor/PhAT/Private/PhATRender.cpp
@@ -438,7 +438,7 @@ FColor UPhATEdSkeletalMeshComponent::GetPrimitiveColor(int32 BodyIndex, EKCollis
 		}
 	}
 
-	if (SharedData->EditingMode == FPhATSharedData::PEM_ConstraintEdit)
+	if (!SharedData->bRunningSimulation && SharedData->EditingMode == FPhATSharedData::PEM_ConstraintEdit)
 	{
 		return BoneUnselectedColor;
 	}
@@ -464,7 +464,7 @@ FColor UPhATEdSkeletalMeshComponent::GetPrimitiveColor(int32 BodyIndex, EKCollis
 		return ElemSelectedBodyColor;
 	}
 	
-	if (SharedData->bShowFixedStatus)
+	if (SharedData->bShowFixedStatus && SharedData->bRunningSimulation)
 	{
 		const bool bIsSimulatedAtAll = BodySetup->PhysicsType == PhysType_Simulated || (BodySetup->PhysicsType == PhysType_Default && SharedData->EditorSimOptions->PhysicsBlend > 0.f);
 		if (!bIsSimulatedAtAll)
