Fixed issue where changing the mesh rendering mode in PhAT would not update in the viewport until the second click.

diff --git a/Engine/Source/Editor/PhAT/Private/PhAT.cpp b/Engine/Source/Editor/PhAT/Private/PhAT.cpp
--- a/Engine/Source/Editor/PhAT/Private/PhAT.cpp
+++ b/Engine/Source/Editor/PhAT/Private/PhAT.cpp
@@ -202,6 +202,7 @@ void FPhAT::InitPhAT(const EToolkitMode::Type Mode, const TSharedPtr< class IToo
 	SelectedAnimation = NULL;
 	SelectedSimulation = false;
 	BeforeSimulationWidgetMode = FWidget::EWidgetMode::WM_None;
+	TickCountUntilViewportRefresh = 0;
 
 	SharedData = MakeShareable(new FPhATSharedData);
 
@@ -2461,6 +2462,9 @@ void FPhAT::OnMeshRenderingMode(FPhATSharedData::EPhATRenderMode Mode)
 		SharedData->ConstraintEdit_MeshViewMode = Mode;
 	}
 
+	// Changing the mesh rendering mode requires the skeletal mesh component to change its render state, which is an operation
+	// which is deferred until after render. Hence we need to trigger another viewport refresh on the following frame.
+	TickCountUntilViewportRefresh = 2;
 	RefreshPreviewViewport();
 }
 
@@ -3261,6 +3265,14 @@ void FPhAT::RecordAnimation()
 
 void FPhAT::Tick(float DeltaTime)
 {
+	if (TickCountUntilViewportRefresh > 0)
+	{
+		TickCountUntilViewportRefresh--;
+		if (TickCountUntilViewportRefresh == 0)
+		{
+			RefreshPreviewViewport();
+		}
+	}
 }
 
 TStatId FPhAT::GetStatId() const
diff --git a/Engine/Source/Editor/PhAT/Private/PhAT.h b/Engine/Source/Editor/PhAT/Private/PhAT.h
--- a/Engine/Source/Editor/PhAT/Private/PhAT.h
+++ b/Engine/Source/Editor/PhAT/Private/PhAT.h
@@ -252,6 +252,9 @@ private:
 	/** Preview Viewport */
 	TSharedPtr<SPhATPreviewViewport> PreviewViewport;
 
+	/** Ticks until forcing viewport refresh */
+	int32 TickCountUntilViewportRefresh;
+
 	/** Properties Tab */
 	TSharedPtr<class IDetailsView> Properties;
 
