Improve phat toolbar: Added lit mode, cleaned up camera menu so that it matches the rest of the editor. Default to 3 camera speed

#UE-16638

diff --git a/Engine/Source/Editor/PhAT/Private/PhATPreviewViewportClient.cpp b/Engine/Source/Editor/PhAT/Private/PhATPreviewViewportClient.cpp
--- a/Engine/Source/Editor/PhAT/Private/PhATPreviewViewportClient.cpp
+++ b/Engine/Source/Editor/PhAT/Private/PhATPreviewViewportClient.cpp
@@ -86,6 +86,8 @@ FPhATEdPreviewViewportClient::FPhATEdPreviewViewportClient(TWeakPtr<FPhAT> InPhA
 
 	SetViewModes(VMI_Lit, VMI_Lit);
 
+	SetCameraSpeedSetting(3);
+
 	bUsingOrbitCamera = true;
 
 	if (!FPhAT::IsPIERunning())
@@ -124,7 +126,7 @@ void FPhATEdPreviewViewportClient::DrawCanvas( FViewport& InViewport, FSceneView
 		FText::AsNumber(SharedData->PhysicsAsset->ConstraintSetup.Num()) ).ToString();
 
 	TextItem.Text = FText::FromString( StatusString );
-	Canvas.DrawItem( TextItem, XOffset, 3 );
+	//Canvas.DrawItem( TextItem, XOffset, 3 );	//NOTE: Turning off as I'm not sure how useful this text is, but there seems to be a lot of information on the screen already
 	
 	TextItem.Text = FText::GetEmpty();
 	if (SharedData->bRunningSimulation)
diff --git a/Engine/Source/Editor/PhAT/Private/SPhATPreviewToolbar.cpp b/Engine/Source/Editor/PhAT/Private/SPhATPreviewToolbar.cpp
--- a/Engine/Source/Editor/PhAT/Private/SPhATPreviewToolbar.cpp
+++ b/Engine/Source/Editor/PhAT/Private/SPhATPreviewToolbar.cpp
@@ -10,6 +10,8 @@
 #include "SPhATPreviewToolbar.h"
 #include "Editor/UnrealEd//Public/SEditorViewport.h"
 
+#define LOCTEXT_NAMESPACE "PhatViewportToolBar"
+
 
 void SPhATPreviewViewportToolBar::Construct(const FArguments& InArgs, TSharedPtr<class SEditorViewport> InRealViewport)
 {
@@ -26,16 +28,27 @@ void SPhATPreviewViewportToolBar::Construct(const FArguments& InArgs, TSharedPtr
 		.ForegroundColor( FEditorStyle::GetSlateColor(DefaultForegroundName) )
 		[
 			SNew(SHorizontalBox)
-			+ SHorizontalBox::Slot()
+
+			// Camera Type (Perspective/Top/etc...)
+			+SHorizontalBox::Slot()
 			.AutoWidth()
-			.Padding(5.0f, 2.0f)
+			.Padding(2.0f, 2.0f)
 			[
 				SNew(SEditorViewportToolbarMenu)
 				.ParentToolBar(SharedThis(this))
-				.Cursor(EMouseCursor::Default)
-				.Label(NSLOCTEXT("PhAT", "CameraMenuTitle_Default", "Camera"))
+				.Label(this, &SPhATPreviewViewportToolBar::GetCameraMenuLabel)
+				.LabelIcon(this, &SPhATPreviewViewportToolBar::GetCameraMenuLabelIcon)
 				.OnGetMenuContent(this, &SPhATPreviewViewportToolBar::GeneratePerspectiveMenu)
 			]
+
+			// View menu (lit, unlit, etc...)
+			+ SHorizontalBox::Slot()
+			.AutoWidth()
+			.Padding(5.0f, 2.0f)
+			[
+				SNew(SEditorViewportViewMenu, InRealViewport.ToSharedRef(), SharedThis(this))
+			]
+
 			+SHorizontalBox::Slot()
 			.AutoWidth()
 			.Padding(5.0f, 2.0f)
@@ -43,9 +56,10 @@ void SPhATPreviewViewportToolBar::Construct(const FArguments& InArgs, TSharedPtr
 				SNew(SEditorViewportToolbarMenu)
 				.ParentToolBar(SharedThis(this))
 				.Cursor(EMouseCursor::Default)
-				.Label(NSLOCTEXT("PhAT", "ViewMenuTitle_Default", "View"))
-				.OnGetMenuContent(this, &SPhATPreviewViewportToolBar::GenerateViewMenu) 
+				.Label(NSLOCTEXT("PhAT", "ShowMenuTitle_Default", "Show"))
+				.OnGetMenuContent(this, &SPhATPreviewViewportToolBar::GenerateShowMenu) 
 			]
+
 			+SHorizontalBox::Slot()
 			.AutoWidth()
 			.Padding(5.0f, 2.0f)
@@ -71,6 +85,92 @@ void SPhATPreviewViewportToolBar::Construct(const FArguments& InArgs, TSharedPtr
 	SViewportToolBar::Construct(SViewportToolBar::FArguments());
 }
 
+FText SPhATPreviewViewportToolBar::GetCameraMenuLabel() const
+{
+	FText Label = LOCTEXT("Viewport_Default", "Camera");
+	TSharedPtr<SPhATPreviewViewport> EditorViewport = PhATPtr.Pin()->GetPreviewViewportWidget();
+	if (EditorViewport.IsValid())
+	{
+		switch (EditorViewport->GetViewportClient()->GetViewportType())
+		{
+		case LVT_Perspective:
+			Label = LOCTEXT("CameraMenuTitle_Perspective", "Perspective");
+			break;
+
+		case LVT_OrthoXY:
+			Label = LOCTEXT("CameraMenuTitle_Top", "Top");
+			break;
+
+		case LVT_OrthoYZ:
+			Label = LOCTEXT("CameraMenuTitle_Left", "Left");
+			break;
+
+		case LVT_OrthoXZ:
+			Label = LOCTEXT("CameraMenuTitle_Front", "Front");
+			break;
+
+		case LVT_OrthoNegativeXY:
+			Label = LOCTEXT("CameraMenuTitle_Bottom", "Bottom");
+			break;
+
+		case LVT_OrthoNegativeYZ:
+			Label = LOCTEXT("CameraMenuTitle_Right", "Right");
+			break;
+
+		case LVT_OrthoNegativeXZ:
+			Label = LOCTEXT("CameraMenuTitle_Back", "Back");
+			break;
+		case LVT_OrthoFreelook:
+			break;
+		}
+	}
+
+	return Label;
+}
+
+const FSlateBrush* SPhATPreviewViewportToolBar::GetCameraMenuLabelIcon() const
+{
+	FName Icon = NAME_None;
+	TSharedPtr<SPhATPreviewViewport> EditorViewport = PhATPtr.Pin()->GetPreviewViewportWidget();
+	if (EditorViewport.IsValid())
+	{
+		switch (EditorViewport->GetViewportClient()->GetViewportType())
+		{
+		case LVT_Perspective:
+			Icon = FName("EditorViewport.Perspective");
+			break;
+
+		case LVT_OrthoXY:
+			Icon = FName("EditorViewport.Top");
+			break;
+
+		case LVT_OrthoYZ:
+			Icon = FName("EditorViewport.Left");
+			break;
+
+		case LVT_OrthoXZ:
+			Icon = FName("EditorViewport.Front");
+			break;
+
+		case LVT_OrthoNegativeXY:
+			Icon = FName("EditorViewport.Bottom");
+			break;
+
+		case LVT_OrthoNegativeYZ:
+			Icon = FName("EditorViewport.Right");
+			break;
+
+		case LVT_OrthoNegativeXZ:
+			Icon = FName("EditorViewport.Back");
+			break;
+		case LVT_OrthoFreelook:
+			break;
+		}
+	}
+
+	return FEditorStyle::GetBrush(Icon);
+}
+
 TSharedRef<SWidget> SPhATPreviewViewportToolBar::GeneratePerspectiveMenu() const
 {
 	const FPhATCommands& Actions = FPhATCommands::Get();
@@ -90,23 +190,23 @@ TSharedRef<SWidget> SPhATPreviewViewportToolBar::GeneratePerspectiveMenu() const
 	return PerspectiveMenuBuilder.MakeWidget();
 }
 
-TSharedRef<SWidget> SPhATPreviewViewportToolBar::GenerateViewMenu() const
+TSharedRef<SWidget> SPhATPreviewViewportToolBar::GenerateShowMenu() const
 {
 	const FPhATCommands& Actions = FPhATCommands::Get();
 		
 	const bool bInShouldCloseWindowAfterMenuSelection = true;
-	FMenuBuilder ViewMenuBuilder(bInShouldCloseWindowAfterMenuSelection, PhATPtr.Pin()->GetToolkitCommands());
+	FMenuBuilder ShowMenuBuilder(bInShouldCloseWindowAfterMenuSelection, PhATPtr.Pin()->GetToolkitCommands());
 	{
-		ViewMenuBuilder.AddMenuEntry(Actions.ShowSkeleton);
-		ViewMenuBuilder.AddMenuEntry(Actions.DrawGroundBox);
-		ViewMenuBuilder.AddMenuEntry(Actions.InstanceProperties);
-		ViewMenuBuilder.AddMenuEntry(Actions.ShowKinematicBodies);
-		ViewMenuBuilder.AddMenuEntry(Actions.ToggleGraphicsHierarchy);
-		ViewMenuBuilder.AddMenuEntry(Actions.ToggleBoneInfuences);
-		ViewMenuBuilder.AddMenuEntry(Actions.ToggleMassProperties);
+		ShowMenuBuilder.AddMenuEntry(Actions.ShowSkeleton);
+		ShowMenuBuilder.AddMenuEntry(Actions.DrawGroundBox);
+		ShowMenuBuilder.AddMenuEntry(Actions.InstanceProperties);
+		ShowMenuBuilder.AddMenuEntry(Actions.ShowKinematicBodies);
+		ShowMenuBuilder.AddMenuEntry(Actions.ToggleGraphicsHierarchy);
+		ShowMenuBuilder.AddMenuEntry(Actions.ToggleBoneInfuences);
+		ShowMenuBuilder.AddMenuEntry(Actions.ToggleMassProperties);
 	}
 
-	return ViewMenuBuilder.MakeWidget();	
+	return ShowMenuBuilder.MakeWidget();	
 }
 
 TSharedRef<SWidget> SPhATPreviewViewportToolBar::GenerateModesMenu() const
@@ -120,19 +220,6 @@ TSharedRef<SWidget> SPhATPreviewViewportToolBar::GenerateModesMenu() const
 		{
 			struct Local
 			{
-				static void BuildMeshRenderModeMenu(FMenuBuilder& Menu)
-				{
-					const FPhATCommands& Commands = FPhATCommands::Get();
-
-					Menu.BeginSection("PhATRenderingMode", NSLOCTEXT("PhAT", "MeshRenderModeHeader", "Mesh Rendering Mode"));
-					{
-						Menu.AddMenuEntry(Commands.MeshRenderingMode_Solid);
-						Menu.AddMenuEntry(Commands.MeshRenderingMode_Wireframe);
-						Menu.AddMenuEntry(Commands.MeshRenderingMode_None);
-					}
-					Menu.EndSection();
-				}
-
 				static void BuildCollisionRenderModeMenu(FMenuBuilder& Menu)
 				{
 					const FPhATCommands& Commands = FPhATCommands::Get();
@@ -162,7 +249,6 @@ TSharedRef<SWidget> SPhATPreviewViewportToolBar::GenerateModesMenu() const
 
 			ModesMenuBuilder.BeginSection("PhATSubMenus");
 			{
-				ModesMenuBuilder.AddSubMenu(NSLOCTEXT("PhAT", "MeshRenderModeSubMenu", "MeshRender Mode"), FText::GetEmpty(), FNewMenuDelegate::CreateStatic(&Local::BuildMeshRenderModeMenu));
 				ModesMenuBuilder.AddSubMenu(NSLOCTEXT("PhAT", "CollisionRenderModeSubMenu", "CollisionRender Mode"), FText::GetEmpty(), FNewMenuDelegate::CreateStatic(&Local::BuildCollisionRenderModeMenu));
 				ModesMenuBuilder.AddSubMenu(NSLOCTEXT("PhAT", "ConstraintRenderModeSubMenu", "ConstraintRender Mode"), FText::GetEmpty(), FNewMenuDelegate::CreateStatic(&Local::BuildConstraintRenderModeMenu));
 			}
@@ -172,3 +258,5 @@ TSharedRef<SWidget> SPhATPreviewViewportToolBar::GenerateModesMenu() const
 
 	return ModesMenuBuilder.MakeWidget();	
 }
+
+#undef LOCTEXT_NAMESPACE
\ No newline at end of file
diff --git a/Engine/Source/Editor/PhAT/Private/SPhATPreviewToolbar.h b/Engine/Source/Editor/PhAT/Private/SPhATPreviewToolbar.h
--- a/Engine/Source/Editor/PhAT/Private/SPhATPreviewToolbar.h
+++ b/Engine/Source/Editor/PhAT/Private/SPhATPreviewToolbar.h
@@ -21,11 +21,14 @@ private:
 	/** Generates the toolbar perspective menu content */
 	TSharedRef<SWidget> GeneratePerspectiveMenu() const;
 	/** Generates the toolbar view menu content */
-	TSharedRef<SWidget> GenerateViewMenu() const;
+	TSharedRef<SWidget> GenerateShowMenu() const;
 
 	/** Generates the toolbar modes menu content */
 	TSharedRef<SWidget> GenerateModesMenu() const;
 
+	FText GetCameraMenuLabel() const;
+	const FSlateBrush* GetCameraMenuLabelIcon() const;
+
 private:
 	/** The viewport that we are in */
 	TWeakPtr<FPhAT> PhATPtr;
