Add mesh render mode back into phat.
Fix rendering of bodies when deselected.
Fix non unity build

diff --git a/Engine/Source/Editor/PhAT/Private/PhATRender.cpp b/Engine/Source/Editor/PhAT/Private/PhATRender.cpp
--- a/Engine/Source/Editor/PhAT/Private/PhATRender.cpp
+++ b/Engine/Source/Editor/PhAT/Private/PhATRender.cpp
@@ -453,13 +453,13 @@ FColor UPhATEdSkeletalMeshComponent::GetPrimitiveColor(int32 BodyIndex, EKCollis
 			bInBody = true;
 		}
 
-		if (Body == SharedData->SelectedBodies[i])
+		if (Body == SharedData->SelectedBodies[i] && !SharedData->bRunningSimulation)
 		{
 			return ElemSelectedColor;
 		}
 	}
 
-	if(bInBody)	//this primitive is in a body that's currently selected, but this primitive itself isn't selected
+	if(bInBody && !SharedData->bRunningSimulation)	//this primitive is in a body that's currently selected, but this primitive itself isn't selected
 	{
 		return ElemSelectedBodyColor;
 	}
@@ -474,7 +474,7 @@ FColor UPhATEdSkeletalMeshComponent::GetPrimitiveColor(int32 BodyIndex, EKCollis
 	}
 	else
 	{
-		if (!SharedData->bRunningSimulation)
+		if (!SharedData->bRunningSimulation && SharedData->SelectedBodies.Num())
 		{
 			// If there is no collision with this body, use 'no collision material'.
 			if (SharedData->NoCollisionBodies.Find(BodyIndex) != INDEX_NONE)
@@ -498,14 +498,14 @@ UMaterialInterface* UPhATEdSkeletalMeshComponent::GetPrimitiveMaterial(int32 Bod
 
 	for(int32 i=0; i< SharedData->SelectedBodies.Num(); ++i)
 	{
-		if (Body == SharedData->SelectedBodies[i])
+		if (Body == SharedData->SelectedBodies[i] && !SharedData->bRunningSimulation)
 		{
 			return bHitTest ? BoneMaterialHit : ElemSelectedMaterial;
 		}
 	}
 
 	// If there is no collision with this body, use 'no collision material'.
-	if (SharedData->NoCollisionBodies.Find(BodyIndex) != INDEX_NONE)
+	if (SharedData->NoCollisionBodies.Find(BodyIndex) != INDEX_NONE && !SharedData->bRunningSimulation)
 	{
 		return bHitTest ? BoneMaterialHit : BoneNoCollisionMaterial;
 	}
diff --git a/Engine/Source/Editor/PhAT/Private/SPhATPreviewToolbar.cpp b/Engine/Source/Editor/PhAT/Private/SPhATPreviewToolbar.cpp
--- a/Engine/Source/Editor/PhAT/Private/SPhATPreviewToolbar.cpp
+++ b/Engine/Source/Editor/PhAT/Private/SPhATPreviewToolbar.cpp
@@ -9,6 +9,7 @@
 #include "PhATActions.h"
 #include "SPhATPreviewToolbar.h"
 #include "Editor/UnrealEd//Public/SEditorViewport.h"
+#include "SPhatPreviewViewport.h"
 
 #define LOCTEXT_NAMESPACE "PhatViewportToolBar"
 
@@ -220,6 +221,19 @@ TSharedRef<SWidget> SPhATPreviewViewportToolBar::GenerateModesMenu() const
 		{
 			struct Local
 			{
+				static void BuildMeshRenderModeMenu(FMenuBuilder& Menu)
+				{
+					const FPhATCommands& Commands = FPhATCommands::Get();
+
+					Menu.BeginSection("PhATRenderingMode", NSLOCTEXT("PhAT", "MeshRenderModeHeader", "Mesh Rendering Mode"));
+					{
+						Menu.AddMenuEntry(Commands.MeshRenderingMode_Solid);
+						Menu.AddMenuEntry(Commands.MeshRenderingMode_Wireframe);
+						Menu.AddMenuEntry(Commands.MeshRenderingMode_None);
+					}
+					Menu.EndSection();
+				}
+
 				static void BuildCollisionRenderModeMenu(FMenuBuilder& Menu)
 				{
 					const FPhATCommands& Commands = FPhATCommands::Get();
@@ -249,6 +263,7 @@ TSharedRef<SWidget> SPhATPreviewViewportToolBar::GenerateModesMenu() const
 
 			ModesMenuBuilder.BeginSection("PhATSubMenus");
 			{
+				ModesMenuBuilder.AddSubMenu(NSLOCTEXT("PhAT", "MeshRenderModeSubMenu", "MeshRender Mode"), FText::GetEmpty(), FNewMenuDelegate::CreateStatic(&Local::BuildMeshRenderModeMenu));
 				ModesMenuBuilder.AddSubMenu(NSLOCTEXT("PhAT", "CollisionRenderModeSubMenu", "CollisionRender Mode"), FText::GetEmpty(), FNewMenuDelegate::CreateStatic(&Local::BuildCollisionRenderModeMenu));
 				ModesMenuBuilder.AddSubMenu(NSLOCTEXT("PhAT", "ConstraintRenderModeSubMenu", "ConstraintRender Mode"), FText::GetEmpty(), FNewMenuDelegate::CreateStatic(&Local::BuildConstraintRenderModeMenu));
 			}
