Case-sensitivity fixes.

diff --git a/Engine/Source/Editor/PhAT/Private/SPhATPreviewToolbar.cpp b/Engine/Source/Editor/PhAT/Private/SPhATPreviewToolbar.cpp
--- a/Engine/Source/Editor/PhAT/Private/SPhATPreviewToolbar.cpp
+++ b/Engine/Source/Editor/PhAT/Private/SPhATPreviewToolbar.cpp
@@ -9,7 +9,7 @@
 #include "PhATActions.h"
 #include "SPhATPreviewToolbar.h"
 #include "Editor/UnrealEd//Public/SEditorViewport.h"
-#include "SPhatPreviewViewport.h"
+#include "SPhATPreviewViewport.h"
 #include "PhATPreviewViewportClient.h"
 
 #define LOCTEXT_NAMESPACE "PhatViewportToolBar"
diff --git a/Engine/Source/Runtime/Engine/Classes/Animation/AnimMetaData.h b/Engine/Source/Runtime/Engine/Classes/Animation/AnimMetaData.h
--- a/Engine/Source/Runtime/Engine/Classes/Animation/AnimMetaData.h
+++ b/Engine/Source/Runtime/Engine/Classes/Animation/AnimMetaData.h
@@ -8,7 +8,7 @@
  */
 
 #pragma once
-#include "AnimMetadata.generated.h"
+#include "AnimMetaData.generated.h"
 
 UCLASS(Blueprintable, abstract, const, MinimalAPI, editinlinenew, hidecategories=Object, collapsecategories)
 class UAnimMetaData : public UObject
