Improve PhAT workflow for switching between constraint mode and body mode

#UE-16705

diff --git a/Engine/Source/Editor/PhAT/Private/PhAT.cpp b/Engine/Source/Editor/PhAT/Private/PhAT.cpp
--- a/Engine/Source/Editor/PhAT/Private/PhAT.cpp
+++ b/Engine/Source/Editor/PhAT/Private/PhAT.cpp
@@ -781,7 +781,9 @@ void FPhAT::ExtendToolbar()
 			//phat edit mode combo
 			FUIAction PhatMode;
 			PhatMode.CanExecuteAction = FCanExecuteAction::CreateSP(Phat, &FPhAT::IsNotSimulation);
-			ToolbarBuilder.BeginSection("PhATMode");
+			ToolbarBuilder.AddToolBarButton(Commands.EditingMode_Body);
+			ToolbarBuilder.AddToolBarButton(Commands.EditingMode_Constraint);
+			/*ToolbarBuilder.BeginSection("PhATMode");
 			{
 				ToolbarBuilder.AddComboButton(
 					PhatMode,
@@ -790,7 +792,7 @@ void FPhAT::ExtendToolbar()
 					TAttribute< FText >::Create( TAttribute< FText >::FGetter::CreateSP( Phat, &FPhAT::GetEditModeToolTip) ),
 					TAttribute< FSlateIcon >::Create( TAttribute< FSlateIcon >::FGetter::CreateSP( Phat, &FPhAT::GetEditModeIcon) )
 					);
-			}
+			}*/
 			ToolbarBuilder.EndSection();
 	
 	
@@ -2224,12 +2226,14 @@ void FPhAT::OnEditingMode(int32 Mode)
 	if (Mode == FPhATSharedData::PEM_BodyEdit)
 	{
 		SharedData->EditingMode = FPhATSharedData::PEM_BodyEdit;
+		SharedData->SetSelectedBodiesFromConstraints();
 		RefreshHierachyTree();
 		SharedData->SetSelectedBody(NULL, true); // Forces properties panel to update...
 	}
 	else
 	{
 		SharedData->EditingMode = FPhATSharedData::PEM_ConstraintEdit;
+		SharedData->SetSelectedConstraintsFromBodies();
 		RefreshHierachyTree();
 		SharedData->SetSelectedConstraint(INDEX_NONE, true);
 	}
diff --git a/Engine/Source/Editor/PhAT/Private/PhATSharedData.cpp b/Engine/Source/Editor/PhAT/Private/PhATSharedData.cpp
--- a/Engine/Source/Editor/PhAT/Private/PhATSharedData.cpp
+++ b/Engine/Source/Editor/PhAT/Private/PhATSharedData.cpp
@@ -546,6 +546,65 @@ void FPhATSharedData::SetSelectedBody(const FSelection* Body, bool bGroupSelect
 	PreviewChangedEvent.Broadcast();
 }
 
+void FPhATSharedData::SetSelectedBodiesFromConstraints()
+{
+	if(SelectedConstraints.Num() == 0)
+	{
+		return;
+	}
+
+	SetSelectedBody(nullptr, false);
+	for (const FSelection& Selection : SelectedConstraints)
+	{
+		UPhysicsConstraintTemplate* ConstraintTemplate = PhysicsAsset->ConstraintSetup[Selection.Index];
+		FConstraintInstance & DefaultInstance = ConstraintTemplate->DefaultInstance;
+
+		for (int32 BodyIdx = 0; BodyIdx < PhysicsAsset->BodySetup.Num(); ++BodyIdx)
+		{
+			UBodySetup* BodySetup = PhysicsAsset->BodySetup[BodyIdx];
+			if (DefaultInstance.JointName == BodySetup->BoneName && BodySetup->AggGeom.GetElementCount() > 0)
+			{
+				FSelection Selection(BodyIdx, KPT_Unknown, 0);
+				int32 PrimIndex = 0;
+				for(int32 GeomType = 0; GeomType < KPT_Unknown; ++GeomType)
+				{
+					if(BodySetup->AggGeom.GetElementCount(GeomType) > 0)
+					{
+						Selection.PrimitiveType = (EKCollisionPrimitiveType)GeomType;
+						break;
+					}
+				}
+				SetSelectedBody(&Selection, true);
+			}
+		}
+	}
+}
+
+void FPhATSharedData::SetSelectedConstraintsFromBodies()
+{
+	if (SelectedBodies.Num() == 0)
+	{
+		return;
+	}
+
+	TSet<int32> TmpSelectedConstraints;	//We could have multiple shapes selected which would cause us to add and remove the same constraint.
+	SetSelectedConstraint(INDEX_NONE, false);
+	for (const FSelection& Selection : SelectedBodies)
+	{
+		UBodySetup* BodySetup = PhysicsAsset->BodySetup[Selection.Index];
+		for(int32 ConstraintIdx = 0; ConstraintIdx < PhysicsAsset->ConstraintSetup.Num(); ++ConstraintIdx)
+		{
+			const UPhysicsConstraintTemplate* ConstraintTemplate = PhysicsAsset->ConstraintSetup[ConstraintIdx]; 
+			if(ConstraintTemplate->DefaultInstance.JointName == BodySetup->BoneName && !TmpSelectedConstraints.Contains(ConstraintIdx))
+			{
+				TmpSelectedConstraints.Add(ConstraintIdx);
+				SetSelectedConstraint(ConstraintIdx, true);
+			}
+		}
+	}
+}
+
+
 void FPhATSharedData::UpdateNoCollisionBodies()
 {
 	NoCollisionBodies.Empty();
@@ -861,6 +920,7 @@ bool FPhATSharedData::WeldSelectedBodies(bool bWeld /* = true */)
 	return true;
 }
 
+
 void FPhATSharedData::InitConstraintSetup(UPhysicsConstraintTemplate* ConstraintSetup, int32 ChildBodyIndex, int32 ParentBodyIndex)
 {
 	check(ConstraintSetup);
diff --git a/Engine/Source/Editor/PhAT/Private/PhATSharedData.h b/Engine/Source/Editor/PhAT/Private/PhATSharedData.h
--- a/Engine/Source/Editor/PhAT/Private/PhATSharedData.h
+++ b/Engine/Source/Editor/PhAT/Private/PhATSharedData.h
@@ -112,6 +112,8 @@ public:
 	void CopyConstraint();
 	void PasteBodyProperties();
 	bool WeldSelectedBodies(bool bWeld = true);
+	void SetSelectedConstraintsFromBodies();
+	void SetSelectedBodiesFromConstraints();
 
 	void Mirror();
 
