Fix shadow variable

diff --git a/Engine/Source/Editor/PhAT/Private/PhATSharedData.cpp b/Engine/Source/Editor/PhAT/Private/PhATSharedData.cpp
--- a/Engine/Source/Editor/PhAT/Private/PhATSharedData.cpp
+++ b/Engine/Source/Editor/PhAT/Private/PhATSharedData.cpp
@@ -564,17 +564,17 @@ void FPhATSharedData::SetSelectedBodiesFromConstraints()
 			UBodySetup* BodySetup = PhysicsAsset->BodySetup[BodyIdx];
 			if (DefaultInstance.JointName == BodySetup->BoneName && BodySetup->AggGeom.GetElementCount() > 0)
 			{
-				FSelection Selection(BodyIdx, KPT_Unknown, 0);
+				FSelection NewSelection(BodyIdx, KPT_Unknown, 0);
 				int32 PrimIndex = 0;
 				for(int32 GeomType = 0; GeomType < KPT_Unknown; ++GeomType)
 				{
 					if(BodySetup->AggGeom.GetElementCount(GeomType) > 0)
 					{
-						Selection.PrimitiveType = (EKCollisionPrimitiveType)GeomType;
+						NewSelection.PrimitiveType = (EKCollisionPrimitiveType)GeomType;
 						break;
 					}
 				}
-				SetSelectedBody(&Selection, true);
+				SetSelectedBody(&NewSelection, true);
 			}
 		}
 	}
