Fix phat focus selection not working when moving in the viewport

#UE-16870

diff --git a/Engine/Source/Editor/PhAT/Private/PhAT.cpp b/Engine/Source/Editor/PhAT/Private/PhAT.cpp
--- a/Engine/Source/Editor/PhAT/Private/PhAT.cpp
+++ b/Engine/Source/Editor/PhAT/Private/PhAT.cpp
@@ -1314,10 +1314,7 @@ void FPhAT::BindCommands()
 		FExecuteAction::CreateSP(this, &FPhAT::OnToggleTwist),
 		FCanExecuteAction::CreateSP(this, &FPhAT::IsNotSimulation));
 
-	ToolkitCommands->MapAction(
-		Commands.FocusOnSelection,
-		FExecuteAction::CreateSP(this, &FPhAT::OnFocusSelection),
-		FCanExecuteAction());
+
 	
 	ToolkitCommands->MapAction(
 		Commands.SelectAllObjects,
diff --git a/Engine/Source/Editor/PhAT/Private/SPhATPreviewViewport.cpp b/Engine/Source/Editor/PhAT/Private/SPhATPreviewViewport.cpp
--- a/Engine/Source/Editor/PhAT/Private/SPhATPreviewViewport.cpp
+++ b/Engine/Source/Editor/PhAT/Private/SPhATPreviewViewport.cpp
@@ -59,6 +59,14 @@ TSharedPtr<SViewport> SPhATPreviewViewport::GetViewportWidget() const
 	return ViewportWidget;
 }
 
+void SPhATPreviewViewport::OnFocusViewportToSelection()
+{
+	if(FPhAT* Phat = PhATPtr.Pin().Get())
+	{
+		Phat->OnFocusSelection();
+	}
+}
+
 TSharedRef<FEditorViewportClient> SPhATPreviewViewport::MakeEditorViewportClient()
 {
 	EditorViewportClient = MakeShareable(new FPhATEdPreviewViewportClient(PhATPtr, PhATPtr.Pin()->GetSharedData(), SharedThis(this)));
diff --git a/Engine/Source/Editor/PhAT/Private/SPhATPreviewViewport.h b/Engine/Source/Editor/PhAT/Private/SPhATPreviewViewport.h
--- a/Engine/Source/Editor/PhAT/Private/SPhATPreviewViewport.h
+++ b/Engine/Source/Editor/PhAT/Private/SPhATPreviewViewport.h
@@ -32,6 +32,8 @@ public:
 	/** Returns true if the viewport is visible */
 	bool IsVisible() const override;
 
+	virtual void OnFocusViewportToSelection() override;
+
 	/** Accessors */
 	TSharedPtr<FSceneViewport> GetViewport() const;
 	TSharedPtr<FPhATEdPreviewViewportClient> GetViewportClient() const;
