Bring back the debug text in phat for bounds, bodies, constraints

diff --git a/Engine/Source/Editor/PhAT/Private/PhATPreviewViewportClient.cpp b/Engine/Source/Editor/PhAT/Private/PhATPreviewViewportClient.cpp
--- a/Engine/Source/Editor/PhAT/Private/PhATPreviewViewportClient.cpp
+++ b/Engine/Source/Editor/PhAT/Private/PhATPreviewViewportClient.cpp
@@ -113,20 +113,20 @@ void FPhATEdPreviewViewportClient::DrawCanvas( FViewport& InViewport, FSceneView
 	float W, H;
 	PhATFont->GetCharSize(TEXT('L'), W, H);
 
-	const float XOffset = 200.0f;
+	const float XOffset = 350.0f;
 
 	FCanvasTextItem TextItem( FVector2D::ZeroVector, FText::GetEmpty(), PhATFont, FLinearColor::White );
 
 	// Write body/constraint count at top.
 	FString StatusString = FText::Format(
-		NSLOCTEXT("UnrealEd", "BodiesConstraints_F", "{0} BODIES  {1} CONSIDERED_FOR_BOUNDS  {2} Ratio  {3} CONSTRAINTS"),
+		NSLOCTEXT("UnrealEd", "BodiesConstraints_F", "{0} Bodies  {1} Considered for bounds  {2} Ratio  {3} Constraints"),
 		FText::AsNumber(SharedData->PhysicsAsset->BodySetup.Num()),
 		FText::AsNumber(SharedData->PhysicsAsset->BoundsBodies.Num()),
 		FText::AsNumber(static_cast<float>(SharedData->PhysicsAsset->BoundsBodies.Num())/static_cast<float>(SharedData->PhysicsAsset->BodySetup.Num())),
 		FText::AsNumber(SharedData->PhysicsAsset->ConstraintSetup.Num()) ).ToString();
 
 	TextItem.Text = FText::FromString( StatusString );
-	//Canvas.DrawItem( TextItem, XOffset, 3 );	//NOTE: Turning off as I'm not sure how useful this text is, but there seems to be a lot of information on the screen already
+	Canvas.DrawItem( TextItem, XOffset, 3 );
 	
 	TextItem.Text = FText::GetEmpty();
 	if (SharedData->bRunningSimulation)
