Expose bUpdateJointsFromAnimation to phat

#UE-17735

diff --git a/Engine/Source/Editor/PhAT/Private/PhATPreviewViewportClient.cpp b/Engine/Source/Editor/PhAT/Private/PhATPreviewViewportClient.cpp
--- a/Engine/Source/Editor/PhAT/Private/PhATPreviewViewportClient.cpp
+++ b/Engine/Source/Editor/PhAT/Private/PhATPreviewViewportClient.cpp
@@ -678,6 +678,7 @@ void FPhATEdPreviewViewportClient::Tick(float DeltaSeconds)
 		// We back up the transforms array now
 		SharedData->EditorSkelComp->AnimationSpaceBases = SharedData->EditorSkelComp->GetSpaceBases();
 		SharedData->EditorSkelComp->SetPhysicsBlendWeight(SharedData->EditorSimOptions->PhysicsBlend);
+		SharedData->EditorSkelComp->bUpdateJointsFromAnimation = SharedData->EditorSimOptions->bUpdateJointsFromAnimation;
 	}
 
 	World->Tick(LEVELTICK_All, DeltaSeconds * SharedData->EditorSimOptions->TimeDilation);
diff --git a/Engine/Source/Editor/UnrealEd/Classes/Preferences/PhATSimOptions.h b/Engine/Source/Editor/UnrealEd/Classes/Preferences/PhATSimOptions.h
--- a/Engine/Source/Editor/UnrealEd/Classes/Preferences/PhATSimOptions.h
+++ b/Engine/Source/Editor/UnrealEd/Classes/Preferences/PhATSimOptions.h
@@ -13,6 +13,10 @@ class UNREALED_API UPhATSimOptions : public UObject
 	UPROPERTY(EditAnywhere, transient, Category=Anim)
 	float PhysicsBlend;
 
+	/** Lets you manually control the physics/animation */
+	UPROPERTY(EditAnywhere, transient, Category = Anim)
+	bool bUpdateJointsFromAnimation;
+
 	/** Time between poking ragdoll and starting to blend back. */
 	UPROPERTY(EditAnywhere, config, Category=Anim)
 	float PokePauseTime;
