FIX UE-17689 PhAT Editor Crashes Engine when Collision Changed, bone tree-view didn't get refreshed so user was able to select non-existing bones
CHANGED renamed typo in ResetBoneCollision

diff --git a/Engine/Source/Editor/PhAT/Private/PhAT.cpp b/Engine/Source/Editor/PhAT/Private/PhAT.cpp
--- a/Engine/Source/Editor/PhAT/Private/PhAT.cpp
+++ b/Engine/Source/Editor/PhAT/Private/PhAT.cpp
@@ -949,7 +949,7 @@ void FPhAT::BindCommands()
 		FCanExecuteAction::CreateSP(this, &FPhAT::IsNotSimulation));
 
 	ToolkitCommands->MapAction(
-		Commands.RestetBoneCollision,
+		Commands.ResetBoneCollision,
 		FExecuteAction::CreateSP(this, &FPhAT::OnResetBoneCollision),
 		FCanExecuteAction::CreateSP(this, &FPhAT::IsSelectedEditBodyMode));
 
@@ -1652,7 +1652,7 @@ TSharedPtr<SWidget> FPhAT::BuildMenuWidgetBody(bool bHierarchy /*= false*/)
 		MenuBuilder.AddMenuEntry( Commands.AddBox );
 		MenuBuilder.AddMenuEntry( Commands.AddSphere );
 		MenuBuilder.AddMenuEntry( Commands.AddSphyl );
-		MenuBuilder.AddMenuEntry( Commands.RestetBoneCollision );
+		MenuBuilder.AddMenuEntry( Commands.ResetBoneCollision );
 		MenuBuilder.EndSection();
 
 		MenuBuilder.BeginSection( "BodyActions", LOCTEXT( "BodyHeader", "Body" ) );
@@ -2200,6 +2200,7 @@ void FPhAT::OnResetBoneCollision()
 
 	SharedData->RefreshPhysicsAssetChange(SharedData->PhysicsAsset);	
 	RefreshPreviewViewport();
+	RefreshHierachyTree();
 }
 
 void FPhAT::OnApplyPhysicalMaterial()
diff --git a/Engine/Source/Editor/PhAT/Private/PhATActions.cpp b/Engine/Source/Editor/PhAT/Private/PhATActions.cpp
--- a/Engine/Source/Editor/PhAT/Private/PhATActions.cpp
+++ b/Engine/Source/Editor/PhAT/Private/PhATActions.cpp
@@ -10,7 +10,7 @@ void FPhATCommands::RegisterCommands()
 {
 	UI_COMMAND(ChangeDefaultMesh, "Change Mesh", "Change Default SkeletalMesh", EUserInterfaceActionType::Button, FInputChord());
 	UI_COMMAND(ResetEntireAsset, "Reset Asset", "Opens Body Creation Settings And Replaces Asset Using The New Settings", EUserInterfaceActionType::Button, FInputChord());
-	UI_COMMAND(RestetBoneCollision, "Reset", "Reset Bone Collision", EUserInterfaceActionType::Button, FInputChord());
+	UI_COMMAND(ResetBoneCollision, "Reset", "Reset Bone Collision", EUserInterfaceActionType::Button, FInputChord());
 	UI_COMMAND(ApplyPhysicalMaterial, "Apply PhysMat", "Applies Currently Selected Physical Material In Content Browser To All Bodies", EUserInterfaceActionType::Button, FInputChord());
 	UI_COMMAND(EditingMode_Body, "Body Mode", "Body Editing Mode", EUserInterfaceActionType::RadioButton, FInputChord());
 	UI_COMMAND(EditingMode_Constraint, "Constraint Mode", "Constraint Editing Mode", EUserInterfaceActionType::RadioButton, FInputChord());
diff --git a/Engine/Source/Editor/PhAT/Private/PhATActions.h b/Engine/Source/Editor/PhAT/Private/PhATActions.h
--- a/Engine/Source/Editor/PhAT/Private/PhATActions.h
+++ b/Engine/Source/Editor/PhAT/Private/PhATActions.h
@@ -20,7 +20,7 @@ public:
 	
 	TSharedPtr<FUICommandInfo> ChangeDefaultMesh;
 	TSharedPtr<FUICommandInfo> ResetEntireAsset;
-	TSharedPtr<FUICommandInfo> RestetBoneCollision;
+	TSharedPtr<FUICommandInfo> ResetBoneCollision;
 	TSharedPtr<FUICommandInfo> ApplyPhysicalMaterial;
 	TSharedPtr<FUICommandInfo> EditingMode_Body;
 	TSharedPtr<FUICommandInfo> EditingMode_Constraint;
