UE-19060 Phat needs an FOV slider

diff --git a/Engine/Source/Editor/PhAT/Private/SPhATPreviewToolbar.cpp b/Engine/Source/Editor/PhAT/Private/SPhATPreviewToolbar.cpp
--- a/Engine/Source/Editor/PhAT/Private/SPhATPreviewToolbar.cpp
+++ b/Engine/Source/Editor/PhAT/Private/SPhATPreviewToolbar.cpp
@@ -30,6 +30,17 @@ void SPhATPreviewViewportToolBar::Construct(const FArguments& InArgs, TSharedPtr
 		.ForegroundColor( FEditorStyle::GetSlateColor(DefaultForegroundName) )
 		[
 			SNew(SHorizontalBox)
+			+ SHorizontalBox::Slot()
+			.AutoWidth()
+			.Padding(2.0f, 2.0f)
+			[
+				SNew(SEditorViewportToolbarMenu)
+				.ParentToolBar(SharedThis(this))
+				.Cursor(EMouseCursor::Default)
+				.Image("EditorViewportToolBar.MenuDropdown")
+				.AddMetaData<FTagMetaData>(FTagMetaData(TEXT("EditorViewportToolBar.MenuDropdown")))
+				.OnGetMenuContent(this, &SPhATPreviewViewportToolBar::GenerateOptionsMenu)
+			]
 
 			// Camera Type (Perspective/Top/etc...)
 			+SHorizontalBox::Slot()
@@ -275,4 +286,62 @@ TSharedRef<SWidget> SPhATPreviewViewportToolBar::GenerateModesMenu() const
 	return ModesMenuBuilder.MakeWidget();	
 }
 
+TSharedRef<SWidget> SPhATPreviewViewportToolBar::GenerateOptionsMenu() const
+{
+
+	const bool bIsPerspective = PhATPtr.Pin()->GetPreviewViewportWidget()->GetViewportClient()->IsPerspective();
+	const bool bInShouldCloseWindowAfterMenuSelection = true;
+	FMenuBuilder OptionsMenuBuilder(bInShouldCloseWindowAfterMenuSelection, PhATPtr.Pin()->GetToolkitCommands());
+	{
+		OptionsMenuBuilder.BeginSection("LevelViewportViewportOptions", LOCTEXT("OptionsMenuHeader", "Viewport Options"));
+		{
+			if (bIsPerspective)
+			{
+				OptionsMenuBuilder.AddWidget(GenerateFOVMenu(), LOCTEXT("FOVAngle", "Field of View (H)"));
+			}
+		}
+		OptionsMenuBuilder.EndSection();		
+	}
+
+	return OptionsMenuBuilder.MakeWidget();
+}
+
+TSharedRef<SWidget> SPhATPreviewViewportToolBar::GenerateFOVMenu() const
+{
+	const float FOVMin = 5.f;
+	const float FOVMax = 170.f;
+
+	return
+		SNew(SBox)
+		.HAlign(HAlign_Right)
+		[
+			SNew(SBox)
+			.Padding(FMargin(4.0f, 0.0f, 0.0f, 0.0f))
+			.WidthOverride(100.0f)
+			[
+				SNew(SSpinBox<float>)
+				.Font(FEditorStyle::GetFontStyle(TEXT("MenuItem.Font")))
+				.MinValue(FOVMin)
+				.MaxValue(FOVMax)
+				.Value(this, &SPhATPreviewViewportToolBar::OnGetFOVValue)
+				.OnValueChanged(this, &SPhATPreviewViewportToolBar::OnFOVValueChanged)
+			]
+		];
+}
+
+float SPhATPreviewViewportToolBar::OnGetFOVValue() const
+{
+	return PhATPtr.Pin()->GetPreviewViewportWidget()->GetViewportClient()->ViewFOV;
+}
+
+void SPhATPreviewViewportToolBar::OnFOVValueChanged(float NewValue)
+{
+	bool bUpdateStoredFOV = true;
+	FPhATEdPreviewViewportClient* ViewportClient = PhATPtr.Pin()->GetPreviewViewportWidget()->GetViewportClient().Get();
+	ViewportClient->FOVAngle = NewValue;
+	ViewportClient->ViewFOV = NewValue;
+	ViewportClient->Invalidate();
+}
+
+
 #undef LOCTEXT_NAMESPACE
diff --git a/Engine/Source/Editor/PhAT/Private/SPhATPreviewToolbar.h b/Engine/Source/Editor/PhAT/Private/SPhATPreviewToolbar.h
--- a/Engine/Source/Editor/PhAT/Private/SPhATPreviewToolbar.h
+++ b/Engine/Source/Editor/PhAT/Private/SPhATPreviewToolbar.h
@@ -26,6 +26,27 @@ private:
 	/** Generates the toolbar modes menu content */
 	TSharedRef<SWidget> GenerateModesMenu() const;
 
+	/**
+	* Generates the toolbar options menu content
+	*
+	* @return The widget containing the options menu content
+	*/
+	TSharedRef<SWidget> GenerateOptionsMenu() const;
+
+
+	/**
+	* Generates the toolbar FOV menu content
+	*
+	* @return The widget containing the FOV menu content
+	*/
+	TSharedRef<SWidget> GenerateFOVMenu() const;
+
+	/** Called by the FOV slider in the perspective viewport to get the FOV value */
+	float OnGetFOVValue() const;
+
+	/** Called when the FOV slider is adjusted in the perspective viewport */
+	void OnFOVValueChanged(float NewValue);
+
 	FText GetCameraMenuLabel() const;
 	const FSlateBrush* GetCameraMenuLabelIcon() const;
 
