Copying //UE4/Dev-Framework to //UE4/Dev-Main (Source: //UE4/Dev-Framework @ 3038004)

diff --git a/Engine/Source/Editor/PhAT/Private/PhATSharedData.cpp b/Engine/Source/Editor/PhAT/Private/PhATSharedData.cpp
--- a/Engine/Source/Editor/PhAT/Private/PhATSharedData.cpp
+++ b/Engine/Source/Editor/PhAT/Private/PhATSharedData.cpp
@@ -13,6 +13,7 @@
 #include "Engine/CollisionProfile.h"
 #include "PhysicsEngine/PhysicsConstraintTemplate.h"
 #include "PhysicsEngine/PhysicsHandleComponent.h"
+#include "PhysicsEngine/PhysicsAsset.h"
 #include "Vehicles/WheeledVehicleMovementComponent.h"
 #include "Engine/StaticMesh.h"
 
@@ -76,6 +77,7 @@ FPhATSharedData::~FPhATSharedData()
 void FPhATSharedData::Initialize()
 {
 	EditorSkelComp = NULL;
+	PhysicalAnimationComponent = nullptr;
 
 	USkeletalMesh * PreviewMesh =  NULL; 
 	FStringAssetReference PreviewMeshStringRef = PhysicsAsset->PreviewSkeletalMesh.ToStringReference();
@@ -105,6 +107,9 @@ void FPhATSharedData::Initialize()
 	// Create SkeletalMeshComponent for rendering skeletal mesh
 	EditorSkelComp = NewObject<UPhATEdSkeletalMeshComponent>();
 	EditorSkelComp->SharedData = this;
+
+	PhysicalAnimationComponent = NewObject<UPhysicalAnimationComponent>();
+	PhysicalAnimationComponent->SetSkeletalMeshComponent(EditorSkelComp);
 	
 	// first disable collision first to avoid creating physics body
 	EditorSkelComp->SetCollisionProfileName(UCollisionProfile::BlockAll_ProfileName);
@@ -120,13 +125,14 @@ void FPhATSharedData::Initialize()
 
 	PreviewScene.AddComponent(EditorSkelComp, FTransform::Identity);
 	PreviewScene.AddComponent(EditorFloorComp, FTransform::Identity);
+	PreviewScene.AddComponent(PhysicalAnimationComponent, FTransform::Identity);
 
 	// Look for body setups with no shapes (how does this happen?).
 	// If we find one- just bang on a default box.
 	bool bFoundEmptyShape = false;
-	for (int32 i = 0; i <PhysicsAsset->BodySetup.Num(); ++i)
+	for (int32 i = 0; i <PhysicsAsset->SkeletalBodySetups.Num(); ++i)
 	{
-		UBodySetup* BodySetup = PhysicsAsset->BodySetup[i];
+		UBodySetup* BodySetup = PhysicsAsset->SkeletalBodySetups[i];
 		if (BodySetup->AggGeom.GetElementCount() == 0)
 		{
 			FKBoxElem BoxElem;
@@ -162,9 +168,9 @@ void FPhATSharedData::Initialize()
 	// If so, put up a warning.
 	TArray<int32> MissingBodyIndices;
 	FString BoneNames;
-	for (int32 i = 0; i <PhysicsAsset->BodySetup.Num(); ++i)
+	for (int32 i = 0; i <PhysicsAsset->SkeletalBodySetups.Num(); ++i)
 	{
-		FName BoneName = PhysicsAsset->BodySetup[i]->BoneName;
+		FName BoneName = PhysicsAsset->SkeletalBodySetups[i]->BoneName;
 		int32 BoneIndex = EditorSkelMesh->RefSkeleton.FindBoneIndex(BoneName);
 		if (BoneIndex == INDEX_NONE)
 		{
@@ -185,7 +191,7 @@ void FPhATSharedData::Initialize()
 			PhysicsAsset->SetFlags(RF_Transactional);
 			PhysicsAsset->Modify();
 
-			// Iterate backwards, as PhysicsAsset->BodySetup is a TArray and UE4 containers don't support remove_if()
+			// Iterate backwards, as PhysicsAsset->SkeletalBodySetups is a TArray and UE4 containers don't support remove_if()
 			for ( int32 i = MissingBodyIndices.Num() - 1; i >= 0; --i )
 			{
 				DeleteBody( MissingBodyIndices[i] );
@@ -250,7 +256,7 @@ void FPhATSharedData::Mirror()
 		{
 			MirrorInfos.AddUninitialized();
 			FMirrorInfo & MirrorInfo = MirrorInfos[MirrorInfos.Num() - 1];
-			MirrorInfo.BoneName = PhysicsAsset->BodySetup[Selection.Index]->BoneName;
+			MirrorInfo.BoneName = PhysicsAsset->SkeletalBodySetups[Selection.Index]->BoneName;
 			MirrorInfo.BodyIndex = Selection.Index;
 			MirrorInfo.ConstraintIndex = PhysicsAsset->FindConstraintIndex(MirrorInfo.BoneName);
 		}
@@ -274,13 +280,13 @@ void FPhATSharedData::Mirror()
 		int32 MirrorBoneIndex = PhysicsAsset->FindMirroredBone(EditorSkelMesh, BoneIndex);
 		if (MirrorBoneIndex != INDEX_NONE)
 		{
-			UBodySetup * SrcBody = PhysicsAsset->BodySetup[MirrorInfo.BodyIndex];
+			UBodySetup * SrcBody = PhysicsAsset->SkeletalBodySetups[MirrorInfo.BodyIndex];
 			const FScopedTransaction Transaction(NSLOCTEXT("PhAT", "MirrorBody", "MirrorBody"));
 			MakeNewBody(MirrorBoneIndex, false);
 
 			int32 MirrorBodyIndex = PhysicsAsset->FindControllingBodyIndex(EditorSkelMesh, MirrorBoneIndex);
 
-			UBodySetup * DestBody = PhysicsAsset->BodySetup[MirrorBodyIndex];
+			UBodySetup * DestBody = PhysicsAsset->SkeletalBodySetups[MirrorBodyIndex];
 			DestBody->Modify();
 			DestBody->CopyBodyPropertiesFrom(SrcBody);
 			FQuat ArtistMirrorConvention(0,0,1,0);   // how Epic Maya artists rig the right and left orientation differently.  todo: perhaps move to cvar 
@@ -418,7 +424,7 @@ void FPhATSharedData::SetSelectedBodyAnyPrim(int32 BodyIndex, bool bGroupSelect
 		return;
 	}
 	
-	UBodySetup* BodySetup = PhysicsAsset->BodySetup[BodyIndex];
+	UBodySetup* BodySetup = PhysicsAsset->SkeletalBodySetups[BodyIndex];
 	check(BodySetup);
 
 	if (BodySetup->AggGeom.SphereElems.Num() > 0)
@@ -490,13 +496,13 @@ void FPhATSharedData::SetSelectedBody(const FSelection* Body, bool bGroupSelect
 	else
 	{
 
-		check(GetSelectedBody() && GetSelectedBody()->Index >= 0 && GetSelectedBody()->Index < PhysicsAsset->BodySetup.Num());
+		check(GetSelectedBody() && GetSelectedBody()->Index >= 0 && GetSelectedBody()->Index < PhysicsAsset->SkeletalBodySetups.Num());
 
 		// Set properties dialog to display selected bone (or bone instance) info.
 		TArray<UObject*> Objs;
 		for(int i=0; i<SelectedBodies.Num(); ++i)
 		{
-			Objs.Add(PhysicsAsset->BodySetup[SelectedBodies[i].Index]);
+			Objs.Add(PhysicsAsset->SkeletalBodySetups[SelectedBodies[i].Index]);
 		}
 		
 		GroupSelectionChangedEvent.Broadcast(Objs);
@@ -539,9 +545,9 @@ void FPhATSharedData::SetSelectedBodiesFromConstraints()
 		UPhysicsConstraintTemplate* ConstraintTemplate = PhysicsAsset->ConstraintSetup[Selection.Index];
 		FConstraintInstance & DefaultInstance = ConstraintTemplate->DefaultInstance;
 
-		for (int32 BodyIdx = 0; BodyIdx < PhysicsAsset->BodySetup.Num(); ++BodyIdx)
+		for (int32 BodyIdx = 0; BodyIdx < PhysicsAsset->SkeletalBodySetups.Num(); ++BodyIdx)
 		{
-			UBodySetup* BodySetup = PhysicsAsset->BodySetup[BodyIdx];
+			UBodySetup* BodySetup = PhysicsAsset->SkeletalBodySetups[BodyIdx];
 			if (DefaultInstance.JointName == BodySetup->BoneName && BodySetup->AggGeom.GetElementCount() > 0)
 			{
 				FSelection NewSelection(BodyIdx, KPT_Unknown, 0);
@@ -571,7 +577,7 @@ void FPhATSharedData::SetSelectedConstraintsFromBodies()
 	SetSelectedConstraint(INDEX_NONE, false);
 	for (const FSelection& Selection : SelectedBodies)
 	{
-		UBodySetup* BodySetup = PhysicsAsset->BodySetup[Selection.Index];
+		UBodySetup* BodySetup = PhysicsAsset->SkeletalBodySetups[Selection.Index];
 		for(int32 ConstraintIdx = 0; ConstraintIdx < PhysicsAsset->ConstraintSetup.Num(); ++ConstraintIdx)
 		{
 			const UPhysicsConstraintTemplate* ConstraintTemplate = PhysicsAsset->ConstraintSetup[ConstraintIdx]; 
@@ -590,10 +596,10 @@ void FPhATSharedData::UpdateNoCollisionBodies()
 	NoCollisionBodies.Empty();
 
 	// Query disable table with selected body and every other body.
-	for (int32 i = 0; i <PhysicsAsset->BodySetup.Num(); ++i)
+	for (int32 i = 0; i <PhysicsAsset->SkeletalBodySetups.Num(); ++i)
 	{
 		// Add any bodies with bNoCollision
-		if (PhysicsAsset->BodySetup[i]->DefaultInstance.GetCollisionEnabled() == ECollisionEnabled::NoCollision)
+		if (PhysicsAsset->SkeletalBodySetups[i]->DefaultInstance.GetCollisionEnabled() == ECollisionEnabled::NoCollision)
 		{
 			NoCollisionBodies.Add(i);
 		}
@@ -602,7 +608,7 @@ void FPhATSharedData::UpdateNoCollisionBodies()
 			// Add this body if it has disabled collision with selected.
 			FRigidBodyIndexPair Key(i, GetSelectedBody()->Index);
 
-			if (PhysicsAsset->BodySetup[GetSelectedBody()->Index]->DefaultInstance.GetCollisionEnabled() == ECollisionEnabled::NoCollision ||
+			if (PhysicsAsset->SkeletalBodySetups[GetSelectedBody()->Index]->DefaultInstance.GetCollisionEnabled() == ECollisionEnabled::NoCollision ||
 				PhysicsAsset->CollisionDisableTable.Find(Key))
 			{
 				NoCollisionBodies.Add(i);
@@ -718,7 +724,7 @@ void FPhATSharedData::CopyBody()
 {
 	check(SelectedBodies.Num() == 1);
 
-	CopiedBodySetup = PhysicsAsset->BodySetup[GetSelectedBody()->Index];
+	CopiedBodySetup = PhysicsAsset->SkeletalBodySetups[GetSelectedBody()->Index];
 }
 
 void FPhATSharedData::PasteBodyProperties()
@@ -742,14 +748,14 @@ void FPhATSharedData::PasteBodyProperties()
 		// Copy setup/instance properties - based on what we are viewing.
 		if (!bShowInstanceProps)
 		{
-			UBodySetup* ToBodySetup = PhysicsAsset->BodySetup[SelectedBodies[i].Index];
+			UBodySetup* ToBodySetup = PhysicsAsset->SkeletalBodySetups[SelectedBodies[i].Index];
 			UBodySetup* FromBodySetup = CopiedBodySetup;
 			ToBodySetup->Modify();
 			ToBodySetup->CopyBodyPropertiesFrom(FromBodySetup);
 		}
 		else
 		{
-			FBodyInstance* ToBodyInstance = &PhysicsAsset->BodySetup[SelectedBodies[i].Index]->DefaultInstance;
+			FBodyInstance* ToBodyInstance = &PhysicsAsset->SkeletalBodySetups[SelectedBodies[i].Index]->DefaultInstance;
 			FBodyInstance* FromBodyInstance = &CopiedBodySetup->DefaultInstance;
 			ToBodyInstance->CopyBodyInstancePropertiesFrom(FromBodyInstance);
 		}
@@ -805,11 +811,11 @@ bool FPhATSharedData::WeldSelectedBodies(bool bWeld /* = true */)
 	const FSelection& Body0 = SelectedBodies[BodyIndex0];
 	const FSelection& Body1 = SelectedBodies[BodyIndex1];
 
-	FName Bone0Name = PhysicsAsset->BodySetup[Body0.Index]->BoneName;
+	FName Bone0Name = PhysicsAsset->SkeletalBodySetups[Body0.Index]->BoneName;
 	int32 Bone0Index = EditorSkelMesh->RefSkeleton.FindBoneIndex(Bone0Name);
 	check(Bone0Index != INDEX_NONE);
 
-	FName Bone1Name = PhysicsAsset->BodySetup[Body1.Index]->BoneName;
+	FName Bone1Name = PhysicsAsset->SkeletalBodySetups[Body1.Index]->BoneName;
 	int32 Bone1Index = EditorSkelMesh->RefSkeleton.FindBoneIndex(Bone1Name);
 	check(Bone1Index != INDEX_NONE);
 
@@ -833,7 +839,7 @@ bool FPhATSharedData::WeldSelectedBodies(bool bWeld /* = true */)
 		ChildPrimitiveType = Body1.PrimitiveType;
 		ParentPrimitiveIndex = Body0.PrimitiveIndex;
 		//Child geoms get appended so just add it. This is kind of a hack but this whole indexing scheme needs to be rewritten anyway
-		ChildPrimitiveIndex = Body1.PrimitiveIndex + PhysicsAsset->BodySetup[Body0.Index]->AggGeom.GetElementCount(ChildPrimitiveType);
+		ChildPrimitiveIndex = Body1.PrimitiveIndex + PhysicsAsset->SkeletalBodySetups[Body0.Index]->AggGeom.GetElementCount(ChildPrimitiveType);
 
 		bCanWeld = true;
 	}else if(PhysicsAsset->FindControllingBodyIndex(EditorSkelMesh, Bone0ParentIndex) == Body1.Index)
@@ -845,7 +851,7 @@ bool FPhATSharedData::WeldSelectedBodies(bool bWeld /* = true */)
 		ChildPrimitiveType = Body0.PrimitiveType;
 		ParentPrimitiveIndex = Body1.PrimitiveIndex;
 		//Child geoms get appended so just add it. This is kind of a hack but this whole indexing scheme needs to be rewritten anyway
-		ChildPrimitiveIndex = Body0.PrimitiveIndex + PhysicsAsset->BodySetup[Body1.Index]->AggGeom.GetElementCount(ChildPrimitiveType);
+		ChildPrimitiveIndex = Body0.PrimitiveIndex + PhysicsAsset->SkeletalBodySetups[Body1.Index]->AggGeom.GetElementCount(ChildPrimitiveType);
 
 		bCanWeld = true;
 	}
@@ -866,8 +872,8 @@ bool FPhATSharedData::WeldSelectedBodies(bool bWeld /* = true */)
 		PhysicsAsset->Modify();
 
 		// .. the parent and child bodies..
-		PhysicsAsset->BodySetup[ParentBodyIndex]->Modify();
-		PhysicsAsset->BodySetup[ChildBodyIndex]->Modify();
+		PhysicsAsset->SkeletalBodySetups[ParentBodyIndex]->Modify();
+		PhysicsAsset->SkeletalBodySetups[ChildBodyIndex]->Modify();
 
 		// .. and any constraints of the 'child' body..
 		TArray<int32>	Constraints;
@@ -907,8 +913,8 @@ void FPhATSharedData::InitConstraintSetup(UPhysicsConstraintTemplate* Constraint
 
 	ConstraintSetup->Modify(false);
 
-	UBodySetup* ChildBodySetup = PhysicsAsset->BodySetup[ ChildBodyIndex ];
-	UBodySetup* ParentBodySetup = PhysicsAsset->BodySetup[ ParentBodyIndex ];
+	UBodySetup* ChildBodySetup = PhysicsAsset->SkeletalBodySetups[ ChildBodyIndex ];
+	UBodySetup* ParentBodySetup = PhysicsAsset->SkeletalBodySetups[ ParentBodyIndex ];
 	check(ChildBodySetup && ParentBodySetup);
 
 	const int32 ChildBoneIndex = EditorSkelMesh->RefSkeleton.FindBoneIndex(ChildBodySetup->BoneName);
@@ -957,7 +963,7 @@ void FPhATSharedData::MakeNewBody(int32 NewBoneIndex, bool bAutoSelect)
 
 	// Create the physics body.
 	NewBodyIndex = FPhysicsAssetUtils::CreateNewBody(PhysicsAsset, NewBoneName);
-	UBodySetup* BodySetup = PhysicsAsset->BodySetup[ NewBodyIndex ];
+	UBodySetup* BodySetup = PhysicsAsset->SkeletalBodySetups[ NewBodyIndex ];
 	check(BodySetup->BoneName == NewBoneName);
 	
 	BodySetup->Modify();
@@ -989,7 +995,7 @@ void FPhATSharedData::MakeNewBody(int32 NewBoneIndex, bool bAutoSelect)
 			// If the child bone is physical, it may require fixing up in regards to constraints
 			if (ChildBodyIndex != INDEX_NONE)
 			{
-				UBodySetup* ChildBody = PhysicsAsset->BodySetup[ ChildBodyIndex ];
+				UBodySetup* ChildBody = PhysicsAsset->SkeletalBodySetups[ ChildBodyIndex ];
 				check(ChildBody);
 
 				int32 ConstraintIndex = PhysicsAsset->FindConstraintIndex(ChildBody->BoneName);
@@ -1183,7 +1189,7 @@ void FPhATSharedData::DeleteBody(int32 DelBodyIndex, bool bRefreshComponent)
 	PhysicsAsset->Modify();
 
 	// .. the body..
-	UBodySetup * BodySetup = PhysicsAsset->BodySetup[DelBodyIndex];
+	UBodySetup * BodySetup = PhysicsAsset->SkeletalBodySetups[DelBodyIndex];
 	BodySetup->Modify();	
 
 	// .. and any constraints to the body.
@@ -1200,7 +1206,7 @@ void FPhATSharedData::DeleteBody(int32 DelBodyIndex, bool bRefreshComponent)
 	{
 		int32 ParentBodyIndex = PhysicsAsset->FindParentBodyIndex(EditorSkelMesh, BoneIndex);
 
-		UBodySetup * ParentBody = ParentBodyIndex != INDEX_NONE ? PhysicsAsset->BodySetup[ParentBodyIndex] : NULL;
+		UBodySetup * ParentBody = ParentBodyIndex != INDEX_NONE ? PhysicsAsset->SkeletalBodySetups[ParentBodyIndex] : NULL;
 
 		for (const int32 ConstraintIndex : Constraints)
 		{
@@ -1212,7 +1218,7 @@ void FPhATSharedData::DeleteBody(int32 DelBodyIndex, bool bRefreshComponent)
 				//for all constraints that contain a nearest child of this body, create a copy of the constraint between the child and parent
 				for (const int32 BodyBelowIndex : NearestBodiesBelow)
 				{
-					UBodySetup * BodyBelow = PhysicsAsset->BodySetup[BodyBelowIndex];
+					UBodySetup * BodyBelow = PhysicsAsset->SkeletalBodySetups[BodyBelowIndex];
 
 					if (Constraint->DefaultInstance.ConstraintBone1 == BodyBelow->BoneName)
 					{
@@ -1260,7 +1266,7 @@ void FPhATSharedData::DeleteCurrentPrim()
 	TArray<UBodySetup*> BodySetups;
 	for(int32 i=0; i<SelectedBodies.Num(); ++i)
 	{
-		UBodySetup* BodySetup = PhysicsAsset->BodySetup[SelectedBodies[i].Index];
+		UBodySetup* BodySetup = PhysicsAsset->SkeletalBodySetups[SelectedBodies[i].Index];
 		BodySelectionMap.FindOrAdd(BodySetup).Add(SelectedBodies[i]);
 	}
 
@@ -1467,7 +1473,7 @@ void FPhATSharedData::ToggleInstanceProperties()
 	{
 		if (GetSelectedBody())
 		{
-			UBodySetup* BodySetup = PhysicsAsset->BodySetup[GetSelectedBody()->Index];
+			UBodySetup* BodySetup = PhysicsAsset->SkeletalBodySetups[GetSelectedBody()->Index];
 
 			// Set properties dialog to display selected bone (or bone instance) info.
 			SelectionChangedEvent.Broadcast(BodySetup, GetSelectedBody());
@@ -1478,7 +1484,7 @@ void FPhATSharedData::ToggleInstanceProperties()
 void FPhATSharedData::ToggleSimulation()
 {
 	// don't start simulation if there are no bodies or if we are manipulating a body
-	if (PhysicsAsset->BodySetup.Num() == 0 || bManipulating)
+	if (PhysicsAsset->SkeletalBodySetups.Num() == 0 || bManipulating)
 	{
 		return;  
 	}
@@ -1499,6 +1505,7 @@ void FPhATSharedData::EnableSimulation(bool bEnableSimulation)
 		EditorSkelComp->SetSimulatePhysics(true);
 		EditorSkelComp->SetPhysicsBlendWeight(EditorSimOptions->PhysicsBlend);
 		EditorSkelComp->InitArticulated(PreviewScene.GetWorld()->GetPhysicsScene());
+		PhysicalAnimationComponent->SetSkeletalMeshComponent(EditorSkelComp);
 
 		// Make it start simulating
 		EditorSkelComp->WakeAllRigidBodies();
@@ -1510,6 +1517,7 @@ void FPhATSharedData::EnableSimulation(bool bEnableSimulation)
 	{
 		// Stop any animation and clear node when stopping simulation.
 		EditorSkelComp->SetAnimation(NULL);
+		PhysicalAnimationComponent->SetSkeletalMeshComponent(nullptr);
 
 		// Turn off/remove the physics instance for this thing, and move back to start location.
 		EditorSkelComp->TermArticulated();
@@ -1569,14 +1577,14 @@ void FPhATSharedData::PostUndo()
 	for (int32 BodyIndex = 0; BodyIndex < SelectedBodies.Num() && bInvalidSelection == false; ++BodyIndex)
 	{
 		const FSelection& Selection = SelectedBodies[BodyIndex];
-		if (PhysicsAsset->BodySetup.Num() <= Selection.Index)
+		if (PhysicsAsset->SkeletalBodySetups.Num() <= Selection.Index)
 		{
 			bInvalidSelection = true;
 		}
 		else
 		{
 			
-			if (UBodySetup * BodySetup = PhysicsAsset->BodySetup[Selection.Index])
+			if (UBodySetup * BodySetup = PhysicsAsset->SkeletalBodySetups[Selection.Index])
 			{
 				switch (Selection.PrimitiveType)
 				{
diff --git a/Engine/Source/Editor/PhAT/Private/PhATSharedData.h b/Engine/Source/Editor/PhAT/Private/PhATSharedData.h
--- a/Engine/Source/Editor/PhAT/Private/PhATSharedData.h
+++ b/Engine/Source/Editor/PhAT/Private/PhATSharedData.h
@@ -168,6 +168,9 @@ public:
 	/** PhAT specific skeletal mesh component */
 	UPhATEdSkeletalMeshComponent* EditorSkelComp;
 
+	/** PhAT specific physical animation component */
+	class UPhysicalAnimationComponent* PhysicalAnimationComponent;
+
 	/** The skeletal mesh being used to preview the physics asset */
 	USkeletalMesh* EditorSkelMesh;
 
diff --git a/Engine/Source/Editor/SequenceRecorder/Private/AnimationRecorder.cpp b/Engine/Source/Editor/SequenceRecorder/Private/AnimationRecorder.cpp
--- a/Engine/Source/Editor/SequenceRecorder/Private/AnimationRecorder.cpp
+++ b/Engine/Source/Editor/SequenceRecorder/Private/AnimationRecorder.cpp
@@ -480,7 +480,7 @@ void FAnimationRecorder::UpdateRecord(USkeletalMeshComponent* Component, float D
 			FBlendedHeapCurve BlendedCurve;
 			if (AnimCurves.Elements.Num() > 0 && PreviousAnimCurves.Elements.Num() == AnimCurves.Elements.Num())
 			{
-				BlendedCurve.Blend(PreviousAnimCurves, AnimCurves, BlendAlpha);
+				BlendedCurve.Lerp(PreviousAnimCurves, AnimCurves, BlendAlpha);
 			}
 			else
 			{
diff --git a/Engine/Source/Editor/UnrealEd/Classes/Factories/FbxAnimSequenceImportData.h b/Engine/Source/Editor/UnrealEd/Classes/Factories/FbxAnimSequenceImportData.h
--- a/Engine/Source/Editor/UnrealEd/Classes/Factories/FbxAnimSequenceImportData.h
+++ b/Engine/Source/Editor/UnrealEd/Classes/Factories/FbxAnimSequenceImportData.h
@@ -59,27 +59,29 @@ class UNREALED_API UFbxAnimSequenceImportData : public UFbxAssetImportData
 	UPROPERTY(EditAnywhere, AdvancedDisplay, config, Category = ImportSettings)
 	bool bImportCustomAttribute;
 
-	/** Set material curve type for the custom attribute*/
-	UPROPERTY(EditAnywhere, AdvancedDisplay, config, Category = ImportSettings, meta = (EditCondition = "bImportCustomAttribute", DisplayName="   Set Material Curve Type"))
+	/** Set Material Curve Type for all custom attributes that exists */
+	UPROPERTY(EditAnywhere, AdvancedDisplay, config, Category = ImportSettings, meta = (EditCondition = "bImportCustomAttribute", DisplayName="Set Material Curve Type"))
 	bool bSetMaterialDriveParameterOnCustomAttribute;
 
-	/** If enabled, this will import a curve within the animation */
-	UPROPERTY(EditAnywhere, AdvancedDisplay, config, Category = ImportSettings)
-	bool bPreserveLocalTransform;
+	/** Set Material Curve Type for the custom attribute with the following suffixes. This doesn't matter if Set Material Curve Type is true  */
+	UPROPERTY(EditAnywhere, AdvancedDisplay, config, Category = ImportSettings, meta = (EditCondition = "bImportCustomAttribute", DisplayName = "Material Curve Suffixes"))
+	TArray<FString> MaterialCurveSuffixes;
+
+	/** When importing custom attribute as curve, remove redundant keys */
+	UPROPERTY(EditAnywhere, AdvancedDisplay, config, Category = ImportSettings, meta = (EditCondition = "bImportCustomAttribute", DisplayName = "Remove Redundant Keys"))
+	bool bRemoveRedundantKeys;
 
 	/** If enabled, this will delete this type of asset from the FBX */
 	UPROPERTY(EditAnywhere, AdvancedDisplay, config, Category = ImportSettings)
 	bool bDeleteExistingMorphTargetCurves;
 
-	UPROPERTY(EditAnywhere, AdvancedDisplay, config, Category = ImportSettings)
-	TArray<FString> MaterialCurveSuffixes;
-
-	UPROPERTY(EditAnywhere, AdvancedDisplay, config, Category = ImportSettings)
-	TArray<FString> PoseCurveSuffixes;
+	/** When importing custom attribute or morphtarget as curve, do not import if it doens't have any value other than zero. This is to avoid adding extra curves to evaluate */
+	UPROPERTY(EditAnywhere, AdvancedDisplay, config, Category = ImportSettings, meta = (DisplayName = "Do not import curves with only 0 values"))
+	bool bDoNotImportCurveWithZero;
 
 	/** If enabled, this will import a curve within the animation */
 	UPROPERTY(EditAnywhere, AdvancedDisplay, config, Category = ImportSettings)
-	bool bRemoveRedundantKeys;
+	bool bPreserveLocalTransform;
 
 	/** Gets or creates fbx import data for the specified anim sequence */
 	static UFbxAnimSequenceImportData* GetImportDataForAnimSequence(UAnimSequence* AnimSequence, UFbxAnimSequenceImportData* TemplateForCreation);
diff --git a/Engine/Source/Editor/UnrealEd/Private/Animation/DebugSkelMeshComponent.cpp b/Engine/Source/Editor/UnrealEd/Private/Animation/DebugSkelMeshComponent.cpp
--- a/Engine/Source/Editor/UnrealEd/Private/Animation/DebugSkelMeshComponent.cpp
+++ b/Engine/Source/Editor/UnrealEd/Private/Animation/DebugSkelMeshComponent.cpp
@@ -405,7 +405,9 @@ void UDebugSkelMeshComponent::GenSpaceBases(TArray<FTransform>& OutSpaceBases)
 	TempLocalAtoms.AddUninitialized(OutSpaceBases.Num());
 	FVector TempRootBoneTranslation;
 	FBlendedHeapCurve TempCurve;
+	PreviewInstance->PrePerformAnimationEvaluation();
 	PerformAnimationEvaluation(SkeletalMesh, AnimScriptInstance, OutSpaceBases, TempLocalAtoms, TempRootBoneTranslation, TempCurve);
+	PreviewInstance->PostPerformAnimationEvaluation();
 }
 
 void UDebugSkelMeshComponent::RefreshBoneTransforms(FActorComponentTickFunction* TickFunction)
diff --git a/Engine/Source/Editor/UnrealEd/Private/Editor.cpp b/Engine/Source/Editor/UnrealEd/Private/Editor.cpp
--- a/Engine/Source/Editor/UnrealEd/Private/Editor.cpp
+++ b/Engine/Source/Editor/UnrealEd/Private/Editor.cpp
@@ -1084,6 +1084,8 @@ namespace EditorUtilities
 				TSet<const UProperty*> SourceUCSModifiedProperties;
 				SourceComponent->GetUCSModifiedProperties(SourceUCSModifiedProperties);
 
+				TArray<UActorComponent*> ComponentInstancesToReregister;
+
 				// Copy component properties
 				for( UProperty* Property = ComponentClass->PropertyLink; Property != nullptr; Property = Property->PropertyLinkNext )
 				{
@@ -1188,10 +1190,13 @@ namespace EditorUtilities
 												}
 											}
 
-											CopySingleProperty( TargetComponent, ComponentArchetypeInstance, Property );
+											if (ComponentArchetypeInstance->IsRegistered())
+											{
+												ComponentArchetypeInstance->UnregisterComponent();
+												ComponentInstancesToReregister.Add(ComponentArchetypeInstance);
+											}
 
-											// Re-register the component with the scene
-											ComponentArchetypeInstance->ReregisterComponent();
+											CopySingleProperty( TargetComponent, ComponentArchetypeInstance, Property );
 										}
 									}
 								}
@@ -1206,6 +1211,11 @@ namespace EditorUtilities
 						}
 					}
 				}
+
+				for (UActorComponent* ModifiedComponentInstance : ComponentInstancesToReregister)
+				{
+					ModifiedComponentInstance->RegisterComponent();
+				}
 			}
 		}
 
diff --git a/Engine/Source/Editor/UnrealEd/Private/Factories/EditorFactories.cpp b/Engine/Source/Editor/UnrealEd/Private/Factories/EditorFactories.cpp
--- a/Engine/Source/Editor/UnrealEd/Private/Factories/EditorFactories.cpp
+++ b/Engine/Source/Editor/UnrealEd/Private/Factories/EditorFactories.cpp
@@ -1654,12 +1654,12 @@ UObject* USoundFactory::FactoryCreateBinary
 		// if the sound already exists, remember the user settings
 		USoundWave* ExistingSound = FindObject<USoundWave>( InParent, *Name.ToString() );
 
-		// TODO - Audio Threading. This needs to be sent to the audio device and wait on stopping the sounds
 		TArray<UAudioComponent*> ComponentsToRestart;
 		FAudioDeviceManager* AudioDeviceManager = GEngine->GetAudioDeviceManager();
 		if (AudioDeviceManager && ExistingSound)
 		{
-			AudioDeviceManager->StopSoundsUsingResource(ExistingSound, ComponentsToRestart);
+			// Will block internally on audio thread completing outstanding commands
+			AudioDeviceManager->StopSoundsUsingResource(ExistingSound, &ComponentsToRestart);
 		}
 
 		bool bUseExistingSettings = bSoundFactorySuppressImportOverwriteDialog;
diff --git a/Engine/Source/Editor/UnrealEd/Private/Fbx/FbxAnimSequenceImportData.cpp b/Engine/Source/Editor/UnrealEd/Private/Fbx/FbxAnimSequenceImportData.cpp
--- a/Engine/Source/Editor/UnrealEd/Private/Fbx/FbxAnimSequenceImportData.cpp
+++ b/Engine/Source/Editor/UnrealEd/Private/Fbx/FbxAnimSequenceImportData.cpp
@@ -7,12 +7,12 @@ UFbxAnimSequenceImportData::UFbxAnimSequenceImportData(const FObjectInitializer&
 	: Super(ObjectInitializer)
 	, bImportCustomAttribute(true)
 	, bRemoveRedundantKeys(true)
+	, bDoNotImportCurveWithZero(true)
 {
 	FrameImportRange.Min = 0;
 	FrameImportRange.Max = 0;
 
 	MaterialCurveSuffixes.Add(TEXT("_mat"));
-	PoseCurveSuffixes.Add(TEXT("_pose"));
 }
 
 UFbxAnimSequenceImportData* UFbxAnimSequenceImportData::GetImportDataForAnimSequence(UAnimSequence* AnimSequence, UFbxAnimSequenceImportData* TemplateForCreation)
diff --git a/Engine/Source/Editor/UnrealEd/Private/Fbx/FbxMainImport.cpp b/Engine/Source/Editor/UnrealEd/Private/Fbx/FbxMainImport.cpp
--- a/Engine/Source/Editor/UnrealEd/Private/Fbx/FbxMainImport.cpp
+++ b/Engine/Source/Editor/UnrealEd/Private/Fbx/FbxMainImport.cpp
@@ -240,10 +240,10 @@ void ApplyImportUIToImportOptions(UFbxImportUI* ImportUI, FBXImportOptions& InOu
 	InOutImportOptions.bPreserveLocalTransform = ImportUI->AnimSequenceImportData->bPreserveLocalTransform;
 	InOutImportOptions.bDeleteExistingMorphTargetCurves = ImportUI->AnimSequenceImportData->bDeleteExistingMorphTargetCurves;
 	InOutImportOptions.bRemoveRedundantKeys = ImportUI->AnimSequenceImportData->bRemoveRedundantKeys;
+	InOutImportOptions.bDoNotImportCurveWithZero = ImportUI->AnimSequenceImportData->bDoNotImportCurveWithZero;
 	InOutImportOptions.bImportCustomAttribute = ImportUI->AnimSequenceImportData->bImportCustomAttribute;
 	InOutImportOptions.bSetMaterialDriveParameterOnCustomAttribute = ImportUI->AnimSequenceImportData->bSetMaterialDriveParameterOnCustomAttribute;
 	InOutImportOptions.MaterialCurveSuffixes = ImportUI->AnimSequenceImportData->MaterialCurveSuffixes;
-	InOutImportOptions.PoseCurveSuffixes = ImportUI->AnimSequenceImportData->PoseCurveSuffixes;
 }
 
 void FImportedMaterialData::AddImportedMaterial( FbxSurfaceMaterial& FbxMaterial, UMaterialInterface& UnrealMaterial )
diff --git a/Engine/Source/Editor/UnrealEd/Private/FbxImporter.h b/Engine/Source/Editor/UnrealEd/Private/FbxImporter.h
--- a/Engine/Source/Editor/UnrealEd/Private/FbxImporter.h
+++ b/Engine/Source/Editor/UnrealEd/Private/FbxImporter.h
@@ -129,8 +129,8 @@ struct FBXImportOptions
 	bool	bImportCustomAttribute;
 	bool	bSetMaterialDriveParameterOnCustomAttribute;
 	bool	bRemoveRedundantKeys;
+	bool	bDoNotImportCurveWithZero;
 	TArray<FString> MaterialCurveSuffixes;
-	TArray<FString>	PoseCurveSuffixes;
 
 	/** This allow to add a prefix to the material name when unreal material get created.	
 	*   This prefix can just modify the name of the asset for materials (i.e. TEXT("Mat"))
diff --git a/Engine/Source/Editor/UnrealEd/Private/Kismet2/ComponentEditorUtils.cpp b/Engine/Source/Editor/UnrealEd/Private/Kismet2/ComponentEditorUtils.cpp
--- a/Engine/Source/Editor/UnrealEd/Private/Kismet2/ComponentEditorUtils.cpp
+++ b/Engine/Source/Editor/UnrealEd/Private/Kismet2/ComponentEditorUtils.cpp
@@ -109,6 +109,10 @@ protected:
 		}
 		else if (AActor* NewActor = Cast<AActor>(NewObject))
 		{
+			if (USceneComponent* RootComponent = NewActor->GetRootComponent())
+			{
+				RootComponent->SetWorldLocationAndRotationNoPhysics(FVector(0.f),FRotator(0.f));
+			}
 			NewActor->GetComponents(ActorComponents);
 		}
 
diff --git a/Engine/Source/Editor/UnrealEd/Private/Kismet2/KismetDebugUtilities.cpp b/Engine/Source/Editor/UnrealEd/Private/Kismet2/KismetDebugUtilities.cpp
--- a/Engine/Source/Editor/UnrealEd/Private/Kismet2/KismetDebugUtilities.cpp
+++ b/Engine/Source/Editor/UnrealEd/Private/Kismet2/KismetDebugUtilities.cpp
@@ -1086,13 +1086,17 @@ FKismetDebugUtilities::EWatchTextResult FKismetDebugUtilities::GetWatchText(FStr
 			if(!PropertyBase && AnimBlueprintGeneratedClass)
 			{
 				// are we linked to an anim graph node?
+				UProperty* LinkedProperty = Property;
 				const UAnimGraphNode_Base* Node = Cast<UAnimGraphNode_Base>(WatchPin->GetOuter());
 				if(Node == nullptr && WatchPin->LinkedTo.Num() > 0)
 				{
-					Node = Cast<UAnimGraphNode_Base>(WatchPin->LinkedTo[0]->GetOuter());
+					const UEdGraphPin* LinkedPin = WatchPin->LinkedTo[0];
+					// When we change Node we *must* change Property, so it's still a sub-element of that.
+					LinkedProperty = FKismetDebugUtilities::FindClassPropertyForPin(Blueprint, LinkedPin);
+					Node = Cast<UAnimGraphNode_Base>(LinkedPin->GetOuter());
 				}
 
-				if(Node)
+				if(Node && LinkedProperty)
 				{
 					UStructProperty* NodeStructProperty = Cast<UStructProperty>(FKismetDebugUtilities::FindClassPropertyForNode(Blueprint, Node));
 					if (NodeStructProperty)
@@ -1102,7 +1106,7 @@ FKismetDebugUtilities::EWatchTextResult FKismetDebugUtilities::GetWatchText(FStr
 							if (NodeProperty == NodeStructProperty)
 							{
 								void* NodePtr = NodeProperty->ContainerPtrToValuePtr<void>(ActiveObject);
-								Property->ExportText_InContainer(/*ArrayElement=*/ 0, /*inout*/ OutWatchText, NodePtr, NodePtr, /*Parent=*/ ActiveObject, PPF_PropertyWindow|PPF_BlueprintDebugView);
+								LinkedProperty->ExportText_InContainer(/*ArrayElement=*/ 0, /*inout*/ OutWatchText, NodePtr, NodePtr, /*Parent=*/ ActiveObject, PPF_PropertyWindow|PPF_BlueprintDebugView);
 								return EWTR_Valid;
 							}
 						}
diff --git a/Engine/Source/Editor/UnrealEd/Private/LevelEditorViewport.cpp b/Engine/Source/Editor/UnrealEd/Private/LevelEditorViewport.cpp
--- a/Engine/Source/Editor/UnrealEd/Private/LevelEditorViewport.cpp
+++ b/Engine/Source/Editor/UnrealEd/Private/LevelEditorViewport.cpp
@@ -3994,12 +3994,8 @@ void FLevelEditorViewportClient::UpdateAudioListener(const FSceneView& View)
 	{
 		if (FAudioDevice* AudioDevice = ViewportWorld->GetAudioDevice())
 		{
-			FReverbSettings ReverbSettings;
-			FInteriorSettings InteriorSettings;
 			const FVector& ViewLocation = GetViewLocation();
 
-			class AAudioVolume* AudioVolume = ViewportWorld->GetAudioSettings(ViewLocation, &ReverbSettings, &InteriorSettings);
-
 			FMatrix CameraToWorld = View.ViewMatrices.ViewMatrix.InverseFast();
 			FVector ProjUp = CameraToWorld.TransformVector(FVector(0, 1000, 0));
 			FVector ProjRight = CameraToWorld.TransformVector(FVector(1000, 0, 0));
@@ -4008,8 +4004,7 @@ void FLevelEditorViewportClient::UpdateAudioListener(const FSceneView& View)
 			ListenerTransform.SetTranslation(ViewLocation);
 			ListenerTransform.NormalizeRotation();
 
-			AudioDevice->SetListener(0, ListenerTransform, 0.f, AudioVolume, InteriorSettings);
-			AudioDevice->SetReverbSettings(AudioVolume, ReverbSettings);
+			AudioDevice->SetListener(ViewportWorld, 0, ListenerTransform, 0.f);
 		}
 	}
 }
diff --git a/Engine/Source/Editor/UnrealEd/Private/PhysicsAssetUtils.cpp b/Engine/Source/Editor/UnrealEd/Private/PhysicsAssetUtils.cpp
--- a/Engine/Source/Editor/UnrealEd/Private/PhysicsAssetUtils.cpp
+++ b/Engine/Source/Editor/UnrealEd/Private/PhysicsAssetUtils.cpp
@@ -167,7 +167,7 @@ bool CreateFromSkeletalMeshInternal(UPhysicsAsset* PhysicsAsset, USkeletalMesh*
 		{
 			// Go ahead and make this bone physical.
 			int32 NewBodyIndex = CreateNewBody(PhysicsAsset, BoneName);
-			UBodySetup* bs = PhysicsAsset->BodySetup[NewBodyIndex];
+			UBodySetup* bs = PhysicsAsset->SkeletalBodySetups[NewBodyIndex];
 			check(bs->BoneName == BoneName);
 
 			// Fill in collision info for this bone.
@@ -211,7 +211,7 @@ bool CreateFromSkeletalMeshInternal(UPhysicsAsset* PhysicsAsset, USkeletalMesh*
 		}
 	}
 
-	return PhysicsAsset->BodySetup.Num() > 0;
+	return PhysicsAsset->SkeletalBodySetups.Num() > 0;
 }
 
 bool CreateFromSkeletalMesh(UPhysicsAsset* PhysicsAsset, USkeletalMesh* SkelMesh, FPhysAssetCreateParams& Params, FText& OutErrorMessage)
@@ -559,13 +559,13 @@ void WeldBodies(UPhysicsAsset* PhysAsset, int32 BaseBodyIndex, int32 AddBodyInde
 		return;
 	}
 
-	UBodySetup* Body1 = PhysAsset->BodySetup[BaseBodyIndex];
+	UBodySetup* Body1 = PhysAsset->SkeletalBodySetups[BaseBodyIndex];
 	int32 Bone1Index = SkelComp->SkeletalMesh->RefSkeleton.FindBoneIndex(Body1->BoneName);
 	check(Bone1Index != INDEX_NONE);
 	FTransform Bone1TM = SkelComp->GetBoneTransform(Bone1Index);
 	Bone1TM.RemoveScaling();
 
-	UBodySetup* Body2 = PhysAsset->BodySetup[AddBodyIndex];
+	UBodySetup* Body2 = PhysAsset->SkeletalBodySetups[AddBodyIndex];
 	int32 Bone2Index = SkelComp->SkeletalMesh->RefSkeleton.FindBoneIndex(Body2->BoneName);
 	check(Bone2Index != INDEX_NONE);
 	FTransform Bone2TM = SkelComp->GetBoneTransform(Bone2Index);
@@ -610,7 +610,7 @@ void WeldBodies(UPhysicsAsset* PhysAsset, int32 BaseBodyIndex, int32 AddBodyInde
 	// We need to update the collision disable table to shift any pairs that included body2 to include body1 instead.
 	// We remove any pairs that include body2 & body1.
 
-	for(int32 i=0; i<PhysAsset->BodySetup.Num(); i++)
+	for(int32 i=0; i<PhysAsset->SkeletalBodySetups.Num(); i++)
 	{
 		if(i == AddBodyIndex) 
 			continue;
@@ -726,13 +726,13 @@ int32 CreateNewBody(UPhysicsAsset* PhysAsset, FName InBodyName)
 		return BodyIndex; // if we already have one for this name - just return that.
 	}
 
-	UBodySetup* NewBodySetup = NewObject<UBodySetup>(PhysAsset, NAME_None, RF_Transactional);
+	USkeletalBodySetup* NewBodySetup = NewObject<USkeletalBodySetup>(PhysAsset, NAME_None, RF_Transactional);
 	// make default to be use complex as simple 
 	NewBodySetup->CollisionTraceFlag = CTF_UseSimpleAsComplex;
 	// newly created bodies default to simulating
 	NewBodySetup->PhysicsType = PhysType_Default;
 
-	int32 BodySetupIndex = PhysAsset->BodySetup.Add( NewBodySetup );
+	int32 BodySetupIndex = PhysAsset->SkeletalBodySetups.Add( NewBodySetup );
 	NewBodySetup->BoneName = InBodyName;
 
 	PhysAsset->UpdateBodySetupIndexMap();
@@ -751,7 +751,7 @@ void DestroyBody(UPhysicsAsset* PhysAsset, int32 bodyIndex)
 	// All elements which refer to a body with index >bodyIndex are adjusted. 
 
 	TMap<FRigidBodyIndexPair,bool> NewCDT;
-	for(int32 i=1; i<PhysAsset->BodySetup.Num(); i++)
+	for(int32 i=1; i<PhysAsset->SkeletalBodySetups.Num(); i++)
 	{
 		for(int32 j=0; j<i; j++)
 		{
@@ -786,7 +786,7 @@ void DestroyBody(UPhysicsAsset* PhysAsset, int32 bodyIndex)
 	}
 
 	// Remove pointer from array. Actual objects will be garbage collected.
-	PhysAsset->BodySetup.RemoveAt(bodyIndex);
+	PhysAsset->SkeletalBodySetups.RemoveAt(bodyIndex);
 
 	PhysAsset->UpdateBodySetupIndexMap();
 	// Update body indices.
diff --git a/Engine/Source/Editor/UnrealEd/Private/PlayLevel.cpp b/Engine/Source/Editor/UnrealEd/Private/PlayLevel.cpp
--- a/Engine/Source/Editor/UnrealEd/Private/PlayLevel.cpp
+++ b/Engine/Source/Editor/UnrealEd/Private/PlayLevel.cpp
@@ -62,6 +62,8 @@
 #include "GameDelegates.h"
 #include "GeneralProjectSettings.h"
 
+#include "AudioThread.h"
+
 DEFINE_LOG_CATEGORY_STATIC(LogPlayLevel, Log, All);
 DEFINE_LOG_CATEGORY_STATIC(LogHMD, Log, All);
 
@@ -572,10 +574,10 @@ void UEditorEngine::TeardownPlaySession(FWorldContext &PieWorldContext)
 	// Stop all audio and remove references to temp level.
 	if (FAudioDevice* AudioDevice = PlayWorld->GetAudioDevice())
 	{
-		AudioDevice->Flush( PlayWorld );
+		AudioDevice->Flush(PlayWorld);
 		AudioDevice->ResetInterpolation();
-		AudioDevice->OnEndPIE(bIsSimulatingInEditor);
-		AudioDevice->TransientMasterVolume = 1.0f;
+		AudioDevice->OnEndPIE(false); // TODO: Should this have been bWasSimulatingInEditor?
+		AudioDevice->SetTransientMasterVolume(1.0f);
 	}
 
 	// Clean up all streaming levels
@@ -2283,19 +2285,19 @@ void UEditorEngine::PlayInEditor( UWorld* InWorld, bool bInSimulateInEditor )
 	GEngine->ClearOnScreenDebugMessages();
 
 	// Flush all audio sources from the editor world
-	FAudioDevice* AudioDevice = EditorWorld->GetAudioDevice();
-	if (AudioDevice)
+	if (FAudioDevice* AudioDevice = EditorWorld->GetAudioDevice())
 	{
-		AudioDevice->Flush( EditorWorld );
+		const bool bEnableSound = PlayInSettings->EnableSound;
+		AudioDevice->Flush(EditorWorld);
 		AudioDevice->ResetInterpolation();
 		AudioDevice->OnBeginPIE(bInSimulateInEditor);
-	}
-	EditorWorld->bAllowAudioPlayback = false;
 
-	if (!PlayInSettings->EnableSound && AudioDevice)
-	{
-		AudioDevice->TransientMasterVolume = 0.0f;
+		if (!bEnableSound)
+		{
+			AudioDevice->SetTransientMasterVolume(0.0f);
+		}
 	}
+	EditorWorld->bAllowAudioPlayback = false;
 
 	if (!GEditor->bAllowMultiplePIEWorlds)
 	{
diff --git a/Engine/Source/Editor/UnrealEd/Private/SkeletalMeshEdit.cpp b/Engine/Source/Editor/UnrealEd/Private/SkeletalMeshEdit.cpp
--- a/Engine/Source/Editor/UnrealEd/Private/SkeletalMeshEdit.cpp
+++ b/Engine/Source/Editor/UnrealEd/Private/SkeletalMeshEdit.cpp
@@ -988,6 +988,29 @@ bool UnFbx::FFbxImporter::ImportCurveToAnimSequence(class UAnimSequence * Target
 	return false;
 }
 
+bool ShouldImportCurve(FbxAnimCurve* Curve, bool bDoNotImportWithZeroValues)
+{
+	if (Curve && Curve->KeyGetCount() > 0)
+	{
+		if (bDoNotImportWithZeroValues)
+		{
+			for (int32 KeyIndex = 0; KeyIndex < Curve->KeyGetCount(); ++KeyIndex)
+			{
+				if (!FMath::IsNearlyZero(Curve->KeyGetValue(KeyIndex)))
+				{
+					return true;
+				}
+			}
+		}
+		else
+		{
+			return true;
+		}
+	}
+
+	return false;
+}
+
 bool UnFbx::FFbxImporter::ImportAnimation(USkeleton* Skeleton, UAnimSequence * DestSeq, const FString& FileName, TArray<FbxNode*>& SortedLinks, TArray<FbxNode*>& NodeArray, FbxAnimStack* CurAnimStack, const int32 ResampleRate, const FbxTimeSpan AnimTimeSpan)
 {
 	// @todo : the length might need to change w.r.t. sampling keys
@@ -1013,7 +1036,7 @@ bool UnFbx::FFbxImporter::ImportAnimation(USkeleton* Skeleton, UAnimSequence * D
 		for (int32 CurveIdx=0; CurveIdx<DestSeq->RawCurveData.FloatCurves.Num(); ++CurveIdx)
 		{
 			auto& Curve = DestSeq->RawCurveData.FloatCurves[CurveIdx];
-			if (Curve.GetCurveTypeFlag(ACF_DrivesMorphTarget))
+			if (Curve.GetCurveTypeFlag(ACF_DriveMorphTarget))
 			{
 				DestSeq->RawCurveData.FloatCurves.RemoveAt(CurveIdx, 1, false);
 				--CurveIdx;
@@ -1059,14 +1082,18 @@ bool UnFbx::FFbxImporter::ImportAnimation(USkeleton* Skeleton, UAnimSequence * D
 							}
 
 							FbxAnimCurve* Curve = Geometry->GetShapeChannel(BlendShapeIndex, ChannelIndex, (FbxAnimLayer*)CurAnimStack->GetMember(0));
-							if (Curve && Curve->KeyGetCount() > 0)
+							if (ShouldImportCurve(Curve, ImportOptions->bDoNotImportCurveWithZero))
 							{
 								FFormatNamedArguments Args;
 								Args.Add(TEXT("BlendShape"), FText::FromString(ChannelName));
 								const FText StatusUpate = FText::Format(LOCTEXT("ImportingMorphTargetCurvesDetail", "Importing Morph Target Curves [{BlendShape}]"), Args);
 								GWarn->StatusUpdate(NodeIndex + 1, NodeArray.Num(), StatusUpate);
 								// now see if we have one already exists. If so, just overwrite that. if not, add new one. 
-								ImportCurveToAnimSequence(DestSeq, *ChannelName, Curve,  ACF_DrivesMorphTarget | ACF_TriggerEvent, AnimTimeSpan, 0.01f /** for some reason blend shape values are coming as 100 scaled **/);
+								ImportCurveToAnimSequence(DestSeq, *ChannelName, Curve,  ACF_DriveMorphTarget | ACF_DriveAttribute, AnimTimeSpan, 0.01f /** for some reason blend shape values are coming as 100 scaled **/);
+							}
+							else
+							{
+								UE_LOG(LogFbx, Warning, TEXT("CurveName(%s) is skipped because it only contains invalid values."), *ChannelName);
 							}
 						}
 					}
@@ -1103,7 +1130,7 @@ bool UnFbx::FFbxImporter::ImportAnimation(USkeleton* Skeleton, UAnimSequence * D
 						FbxAnimCurve * AnimCurve = CurveNode->GetCurve(ChannelIndex);
 						FString ChannelName = CurveNode->GetChannelName(ChannelIndex).Buffer();
 
-						if (AnimCurve)
+						if (ShouldImportCurve(AnimCurve, ImportOptions->bDoNotImportCurveWithZero))
 						{
 							FString FinalCurveName;
 							if (TotalCount == 1)
@@ -1124,7 +1151,7 @@ bool UnFbx::FFbxImporter::ImportAnimation(USkeleton* Skeleton, UAnimSequence * D
 							// first let them override material curve if required
 							if (ImportOptions->bSetMaterialDriveParameterOnCustomAttribute)
 							{
-								CurveFlags |= ACF_DrivesMaterial;
+								CurveFlags |= ACF_DriveMaterial;
 							}
 							else
 							{
@@ -1134,26 +1161,18 @@ bool UnFbx::FFbxImporter::ImportAnimation(USkeleton* Skeleton, UAnimSequence * D
 									int32 TotalSuffix = Suffix.Len();
 									if (CurveName.Right(TotalSuffix) == Suffix)
 									{
-										CurveFlags |= ACF_DrivesMaterial;
+										CurveFlags |= ACF_DriveMaterial;
 										break;
 									}
 								}
 							}
 
-							// apply naming convention for pose curve
-							for (const auto& Suffix : ImportOptions->PoseCurveSuffixes)
-							{
-								int32 TotalSuffix = Suffix.Len();
-								if (ChannelName.Right(TotalSuffix) == Suffix)
-								{
-									CurveFlags |= ACF_DrivesPose;
-									break;
-								}
-							}
-
 							ImportCurveToAnimSequence(DestSeq, FinalCurveName, AnimCurve, CurveFlags, AnimTimeSpan);
 						}
-											
+						else
+						{
+							UE_LOG(LogFbx, Log, TEXT("CurveName(%s) is skipped because it only contains invalid values."), *CurveName);
+						}
 					}
 				}
 
diff --git a/Engine/Source/Editor/UnrealEd/Private/ThumbnailHelpers.cpp b/Engine/Source/Editor/UnrealEd/Private/ThumbnailHelpers.cpp
--- a/Engine/Source/Editor/UnrealEd/Private/ThumbnailHelpers.cpp
+++ b/Engine/Source/Editor/UnrealEd/Private/ThumbnailHelpers.cpp
@@ -624,7 +624,7 @@ bool FAnimationSequenceThumbnailScene::SetAnimation(UAnimSequenceBase* InAnimati
 
 	PreviewAnimation = InAnimation;
 
-	if (InAnimation && InAnimation->IsValidToPlay())
+	if (InAnimation)
 	{
 		if (USkeleton* Skeleton = InAnimation->GetSkeleton())
 		{
@@ -636,25 +636,26 @@ bool FAnimationSequenceThumbnailScene::SetAnimation(UAnimSequenceBase* InAnimati
 			{
 				bSetSucessfully = true;
 
-				// Handle posing the mesh at the middle of the animation
-				const float AnimPosition = InAnimation->SequenceLength / 2.f;
-
-				UDebugSkelMeshComponent* MeshComponent = CastChecked<UDebugSkelMeshComponent>(PreviewActor->GetSkeletalMeshComponent());
+				if (InAnimation->IsValidToPlay())
+				{
+					// Handle posing the mesh at the middle of the animation
+					const float AnimPosition = InAnimation->SequenceLength / 2.f;
 
-				MeshComponent->EnablePreview(true, InAnimation);
-				MeshComponent->Play(false);
-				MeshComponent->Stop();
-				MeshComponent->SetPosition(AnimPosition, false);
+					UDebugSkelMeshComponent* MeshComponent = CastChecked<UDebugSkelMeshComponent>(PreviewActor->GetSkeletalMeshComponent());
 
-				UAnimSingleNodeInstance* SingleNodeInstance = PreviewActor->GetSkeletalMeshComponent()->GetSingleNodeInstance();
-				if (SingleNodeInstance)
-				{
-					SingleNodeInstance->UpdateMontageWeightForTimeSkip(AnimPosition);
-				}
+					MeshComponent->EnablePreview(true, InAnimation);
+					MeshComponent->Play(false);
+					MeshComponent->Stop();
+					MeshComponent->SetPosition(AnimPosition, false);
 
-				PreviewActor->GetSkeletalMeshComponent()->RefreshBoneTransforms(nullptr);
+					UAnimSingleNodeInstance* SingleNodeInstance = PreviewActor->GetSkeletalMeshComponent()->GetSingleNodeInstance();
+					if (SingleNodeInstance)
+					{
+						SingleNodeInstance->UpdateMontageWeightForTimeSkip(AnimPosition);
+					}
 
-				FTransform MeshTransform = FTransform::Identity;
+					PreviewActor->GetSkeletalMeshComponent()->RefreshBoneTransforms(nullptr);
+				}
 
 				PreviewActor->SetActorLocation(FVector(0, 0, 0), false);
 				PreviewActor->GetSkeletalMeshComponent()->UpdateBounds();
diff --git a/Engine/Source/Runtime/AIModule/Classes/Navigation/PathFollowingComponent.h b/Engine/Source/Runtime/AIModule/Classes/Navigation/PathFollowingComponent.h
--- a/Engine/Source/Runtime/AIModule/Classes/Navigation/PathFollowingComponent.h
+++ b/Engine/Source/Runtime/AIModule/Classes/Navigation/PathFollowingComponent.h
@@ -76,21 +76,21 @@ namespace FPathFollowingResultFlags
 	const Type OffPath = (1 << 2);
 
 	/** Aborted (EPathFollowingResult::Aborted) */
-	const Type ExternalCancel = (1 << 3);
+	const Type UserAbort = (1 << 3);
 
-	/** External cancel: owner no longer wants to move */
+	/** Abort details: owner no longer wants to move */
 	const Type OwnerFinished = (1 << 4);
 
-	/** External cancel: path is no longer valid */
+	/** Abort details: path is no longer valid */
 	const Type InvalidPath = (1 << 5);
 
-	/** External cancel: unable to move */
+	/** Abort details: unable to move */
 	const Type MovementStop = (1 << 6);
 
-	/** External cancel: new movement request was received */
+	/** Abort details: new movement request was received */
 	const Type NewRequest = (1 << 7);
 
-	/** External cancel: blueprint MoveTo function was called */
+	/** Abort details: blueprint MoveTo function was called */
 	const Type ForcedScript = (1 << 8);
 
 	/** Finish details: never started, agent was already at goal */
@@ -99,7 +99,7 @@ namespace FPathFollowingResultFlags
 	/** Can be used to create project specific reasons */
 	const Type FirstGameplayFlagShift = 10;
 
-	const Type ExternalCancelFlagMask = ~(Success | Blocked | OffPath);
+	const Type UserAbortFlagMask = ~(Success | Blocked | OffPath);
 
 	FString ToString(uint16 Value);
 }
@@ -117,7 +117,7 @@ struct AIMODULE_API FPathFollowingResult
 
 	bool IsSuccess() const { return HasFlag(FPathFollowingResultFlags::Success); }
 	bool IsFailure() const { return !HasFlag(FPathFollowingResultFlags::Success); }
-	bool IsSkipped() const { return HasFlag(FPathFollowingResultFlags::ExternalCancel | FPathFollowingResultFlags::NewRequest); }
+	bool IsInterrupted() const { return HasFlag(FPathFollowingResultFlags::UserAbort | FPathFollowingResultFlags::NewRequest); }
 	
 	FString ToString() const;
 };
diff --git a/Engine/Source/Runtime/AIModule/Private/Navigation/PathFollowingComponent.cpp b/Engine/Source/Runtime/AIModule/Private/Navigation/PathFollowingComponent.cpp
--- a/Engine/Source/Runtime/AIModule/Private/Navigation/PathFollowingComponent.cpp
+++ b/Engine/Source/Runtime/AIModule/Private/Navigation/PathFollowingComponent.cpp
@@ -27,10 +27,10 @@ DEFINE_LOG_CATEGORY(LogPathFollowing);
 FPathFollowingResult::FPathFollowingResult(uint16 InFlags) : Flags(InFlags)
 {
 	Code =
-		HasFlag(FPathFollowingResultFlags::Success) && !HasFlag(FPathFollowingResultFlags::Blocked | FPathFollowingResultFlags::OffPath | FPathFollowingResultFlags::ExternalCancel) ? EPathFollowingResult::Success :
-		HasFlag(FPathFollowingResultFlags::Blocked) && !HasFlag(FPathFollowingResultFlags::Success | FPathFollowingResultFlags::OffPath | FPathFollowingResultFlags::ExternalCancel) ? EPathFollowingResult::Blocked :
-		HasFlag(FPathFollowingResultFlags::OffPath) && !HasFlag(FPathFollowingResultFlags::Success | FPathFollowingResultFlags::Blocked | FPathFollowingResultFlags::ExternalCancel) ? EPathFollowingResult::OffPath :
-		HasFlag(FPathFollowingResultFlags::ExternalCancel) && !HasFlag(FPathFollowingResultFlags::Success | FPathFollowingResultFlags::Blocked | FPathFollowingResultFlags::OffPath) ? EPathFollowingResult::Aborted :
+		HasFlag(FPathFollowingResultFlags::Success) && !HasFlag(FPathFollowingResultFlags::Blocked | FPathFollowingResultFlags::OffPath | FPathFollowingResultFlags::UserAbort) ? EPathFollowingResult::Success :
+		HasFlag(FPathFollowingResultFlags::Blocked) && !HasFlag(FPathFollowingResultFlags::Success | FPathFollowingResultFlags::OffPath | FPathFollowingResultFlags::UserAbort) ? EPathFollowingResult::Blocked :
+		HasFlag(FPathFollowingResultFlags::OffPath) && !HasFlag(FPathFollowingResultFlags::Success | FPathFollowingResultFlags::Blocked | FPathFollowingResultFlags::UserAbort) ? EPathFollowingResult::OffPath :
+		HasFlag(FPathFollowingResultFlags::UserAbort) && !HasFlag(FPathFollowingResultFlags::Success | FPathFollowingResultFlags::Blocked | FPathFollowingResultFlags::OffPath) ? EPathFollowingResult::Aborted :
 		EPathFollowingResult::Invalid;
 }
 
@@ -40,7 +40,7 @@ FPathFollowingResult::FPathFollowingResult(EPathFollowingResult::Type ResultCode
 		FPathFollowingResultFlags::Success,			// EPathFollowingResult::Success
 		FPathFollowingResultFlags::Blocked,			// EPathFollowingResult::Blocked
 		FPathFollowingResultFlags::OffPath,			// EPathFollowingResult::OffPath
-		FPathFollowingResultFlags::ExternalCancel,	// EPathFollowingResult::Aborted
+		FPathFollowingResultFlags::UserAbort,		// EPathFollowingResult::Aborted
 		0,											// EPathFollowingResult::Skipped_DEPRECATED
 		0,											// EPathFollowingResult::Invalid
 	};
@@ -60,7 +60,7 @@ FString FPathFollowingResultFlags::ToString(uint16 Value)
 		GET_FUNCTION_NAME_STRING_CHECKED(FPathFollowingResultFlags, Success),
 		GET_FUNCTION_NAME_STRING_CHECKED(FPathFollowingResultFlags, Blocked),
 		GET_FUNCTION_NAME_STRING_CHECKED(FPathFollowingResultFlags, OffPath),
-		GET_FUNCTION_NAME_STRING_CHECKED(FPathFollowingResultFlags, ExternalCancel),
+		GET_FUNCTION_NAME_STRING_CHECKED(FPathFollowingResultFlags, UserAbort),
 		GET_FUNCTION_NAME_STRING_CHECKED(FPathFollowingResultFlags, OwnerFinished),
 		GET_FUNCTION_NAME_STRING_CHECKED(FPathFollowingResultFlags, InvalidPath),
 		GET_FUNCTION_NAME_STRING_CHECKED(FPathFollowingResultFlags, MovementStop),
@@ -397,7 +397,7 @@ void UPathFollowingComponent::AbortMove(const UObject& Instigator, FPathFollowin
 	if ((Status != EPathFollowingStatus::Idle) && RequestID.IsEquivalent(GetCurrentRequestId()))
 	{
 		bStopMovementOnFinish = (VelocityMode == EPathFollowingVelocityMode::Reset);
-		OnPathFinished(FPathFollowingResult(EPathFollowingResult::Aborted, AbortFlags & FPathFollowingResultFlags::ExternalCancelFlagMask));
+		OnPathFinished(FPathFollowingResult(EPathFollowingResult::Aborted, AbortFlags & FPathFollowingResultFlags::UserAbortFlagMask));
 	}
 }
 
diff --git a/Engine/Source/Runtime/AIModule/Private/Tasks/AITask_MoveTo.cpp b/Engine/Source/Runtime/AIModule/Private/Tasks/AITask_MoveTo.cpp
--- a/Engine/Source/Runtime/AIModule/Private/Tasks/AITask_MoveTo.cpp
+++ b/Engine/Source/Runtime/AIModule/Private/Tasks/AITask_MoveTo.cpp
@@ -268,7 +268,7 @@ void UAITask_MoveTo::OnRequestFinished(FAIRequestID RequestID, const FPathFollow
 {
 	if (RequestID == MoveRequestID)
 	{
-		if (Result.IsSkipped() && !Result.HasFlag(FPathFollowingResultFlags::ForcedScript))
+		if (Result.HasFlag(FPathFollowingResultFlags::UserAbort) && Result.HasFlag(FPathFollowingResultFlags::NewRequest) && !Result.HasFlag(FPathFollowingResultFlags::ForcedScript))
 		{
 			UE_VLOG(GetGameplayTasksComponent(), LogGameplayTasks, Log, TEXT("%s> ignoring OnRequestFinished, move was aborted by new request"), *GetName());
 		}
diff --git a/Engine/Source/Runtime/ALAudio/Private/ALAudioDevice.cpp b/Engine/Source/Runtime/ALAudio/Private/ALAudioDevice.cpp
--- a/Engine/Source/Runtime/ALAudio/Private/ALAudioDevice.cpp
+++ b/Engine/Source/Runtime/ALAudio/Private/ALAudioDevice.cpp
@@ -215,13 +215,14 @@ bool FALAudioDevice::InitializeHardware( void )
  */
 void FALAudioDevice::UpdateHardware()
 {
-	if (Listeners.Num() > 0)
+	if (GetListeners().Num() > 0)
 	{
 		// Caches the matrix used to transform a sounds position into local space so we can just look
 		// at the Y component after normalization to determine spatialization.
-		const FVector Up = Listeners[0].GetUp();
-		const FVector Right = Listeners[0].GetFront();
-		InverseTransform = FMatrix(Up, Right, Up ^ Right, Listeners[0].Transform.GetTranslation()).InverseFast();
+		const FListener& Listener = GetListeners()[0];
+		const FVector Up = Listener.GetUp();
+		const FVector Right = Listener.GetFront();
+		InverseTransform = FMatrix(Up, Right, Up ^ Right, Listener.Transform.GetTranslation()).InverseFast();
 	}
 }
 
diff --git a/Engine/Source/Runtime/ALAudio/Private/ALAudioSource.cpp b/Engine/Source/Runtime/ALAudio/Private/ALAudioSource.cpp
--- a/Engine/Source/Runtime/ALAudio/Private/ALAudioSource.cpp
+++ b/Engine/Source/Runtime/ALAudio/Private/ALAudioSource.cpp
@@ -87,7 +87,7 @@ void FALSoundSource::Update( void )
 
 	float Volume = 1.0f;
 
-	if (AudioDevice->bIsDeviceMuted)
+	if (AudioDevice->IsAudioDeviceMuted())
 	{
 		Volume = 0.0f;
 	}
@@ -100,7 +100,7 @@ void FALSoundSource::Update( void )
 			Volume *= 1.25f;
 		}
 		Volume *= FApp::GetVolumeMultiplier();
-		Volume *= AudioDevice->PlatformAudioHeadroom;
+		Volume *= AudioDevice->GetPlatformAudioHeadroom();
 		Volume = FMath::Clamp(Volume, 0.0f, MAX_VOLUME);
 	}
 
diff --git a/Engine/Source/Runtime/Android/AndroidAudio/Private/AndroidAudioSource.cpp b/Engine/Source/Runtime/Android/AndroidAudio/Private/AndroidAudioSource.cpp
--- a/Engine/Source/Runtime/Android/AndroidAudio/Private/AndroidAudioSource.cpp
+++ b/Engine/Source/Runtime/Android/AndroidAudio/Private/AndroidAudioSource.cpp
@@ -417,7 +417,7 @@ void FSLESSoundSource::Update( void )
 		// Emulate the bleed to rear speakers followed by stereo fold down
 		Volume *= 1.25f;
 	}
-	Volume *= AudioDevice->PlatformAudioHeadroom;
+	Volume *= AudioDevice->GetPlatformAudioHeadroom();
 	Volume = FMath::Clamp(Volume, 0.0f, MAX_VOLUME);
 	
 	const float Pitch = FMath::Clamp<float>(WaveInstance->Pitch, MIN_PITCH, MAX_PITCH);
diff --git a/Engine/Source/Runtime/AnimGraphRuntime/Private/AnimNodes/AnimNode_PoseBlendNode.cpp b/Engine/Source/Runtime/AnimGraphRuntime/Private/AnimNodes/AnimNode_PoseBlendNode.cpp
--- a/Engine/Source/Runtime/AnimGraphRuntime/Private/AnimNodes/AnimNode_PoseBlendNode.cpp
+++ b/Engine/Source/Runtime/AnimGraphRuntime/Private/AnimNodes/AnimNode_PoseBlendNode.cpp
@@ -9,49 +9,61 @@
 
 void FAnimNode_PoseBlendNode::Initialize(const FAnimationInitializeContext& Context)
 {
-	FAnimNode_AssetPlayerBase::Initialize(Context);
+	FAnimNode_PoseHandler::Initialize(Context);
 
 	SourcePose.Initialize(Context);
 }
 
 void FAnimNode_PoseBlendNode::CacheBones(const FAnimationCacheBonesContext& Context) 
 {
+	FAnimNode_PoseHandler::CacheBones(Context);
 	SourcePose.CacheBones(Context);
 }
 
 void FAnimNode_PoseBlendNode::UpdateAssetPlayer(const FAnimationUpdateContext& Context)
 {
-	EvaluateGraphExposedInputs.Execute(Context);
+	FAnimNode_PoseHandler::UpdateAssetPlayer(Context);
 	SourcePose.Update(Context);
 }
 
 void FAnimNode_PoseBlendNode::Evaluate(FPoseContext& Output)
 {
-	SourcePose.Evaluate(Output);
+	FPoseContext SourceData(Output);
+	SourcePose.Evaluate(SourceData);
 
-	if ((PoseAsset!= NULL) && (Output.AnimInstanceProxy->IsSkeletonCompatible(PoseAsset->GetSkeleton())))
+	if (CurrentPoseAsset.IsValid() && (Output.AnimInstanceProxy->IsSkeletonCompatible(CurrentPoseAsset->GetSkeleton())))
 	{
+		FPoseContext CurrentPose(Output);
+		check(PoseExtractContext.PoseCurves.Num() == PoseUIDList.Num());
 		// only give pose curve, we don't set any more curve here
-		PoseAsset->GetAnimationPose(Output.Pose, Output.Curve, FAnimExtractContext());
-	}
-	// @else it doesn't do anything - it might contain invalid transform but I don't know if overriding transform is the right way 
-	// because the node might contain incoming transform
-}
+		for (int32 PoseIdx = 0; PoseIdx < PoseUIDList.Num(); ++PoseIdx)
+		{
+			PoseExtractContext.PoseCurves[PoseIdx] = SourceData.Curve.Get(PoseUIDList[PoseIdx]);
+		}
 
-void FAnimNode_PoseBlendNode::OverrideAsset(UAnimationAsset* NewAsset)
-{
-	if(UPoseAsset* NewPoseAsset = Cast<UPoseAsset>(NewAsset))
+		if (CurrentPoseAsset.Get()->GetAnimationPose(CurrentPose.Pose, CurrentPose.Curve, PoseExtractContext))
+		{
+			// once we get it, we have to blend by weight
+			if (CurrentPoseAsset->IsValidAdditive())
+			{
+				Output = SourceData;
+				FAnimationRuntime::AccumulateAdditivePose(Output.Pose, CurrentPose.Pose, Output.Curve, CurrentPose.Curve, 1.f, EAdditiveAnimationType::AAT_LocalSpaceBase);
+			}
+			else
+			{
+				FAnimationRuntime::BlendTwoPosesTogetherPerBone(SourceData.Pose, CurrentPose.Pose, SourceData.Curve, CurrentPose.Curve, BoneBlendWeights, Output.Pose, Output.Curve);
+			}
+		}
+	}
+	else
 	{
-		PoseAsset = NewPoseAsset;
+		Output = SourceData;
 	}
 }
 
 void FAnimNode_PoseBlendNode::GatherDebugData(FNodeDebugData& DebugData)
 {
-	FString DebugLine = DebugData.GetNodeName(this);
-	
-	DebugLine += FString::Printf(TEXT("('%s')"), *GetNameSafe(PoseAsset));
-	DebugData.AddDebugItem(DebugLine, true);
+	FAnimNode_PoseHandler::GatherDebugData(DebugData);
 	SourcePose.GatherDebugData(DebugData.BranchFlow(1.f));
 }
 
diff --git a/Engine/Source/Runtime/AnimGraphRuntime/Private/AnimNodes/AnimNode_PoseByName.cpp b/Engine/Source/Runtime/AnimGraphRuntime/Private/AnimNodes/AnimNode_PoseByName.cpp
--- a/Engine/Source/Runtime/AnimGraphRuntime/Private/AnimNodes/AnimNode_PoseByName.cpp
+++ b/Engine/Source/Runtime/AnimGraphRuntime/Private/AnimNodes/AnimNode_PoseByName.cpp
@@ -9,8 +9,20 @@
 
 void FAnimNode_PoseByName::Initialize(const FAnimationInitializeContext& Context)
 {
-	FAnimNode_AssetPlayerBase::Initialize(Context);
+	FAnimNode_PoseHandler::Initialize(Context);
+}
+
+void FAnimNode_PoseByName::CacheBones(const FAnimationCacheBonesContext& Context) 
+{
+	FAnimNode_PoseHandler::CacheBones(Context);
+}
 
+void FAnimNode_PoseByName::UpdateAssetPlayer(const FAnimationUpdateContext& Context)
+{
+	FAnimNode_PoseHandler::UpdateAssetPlayer(Context);
+
+	// for now we do this in every frame because we don't detect name change
+	// this has to happen whenever name changes
 	if (PoseName != NAME_None)
 	{
 		USkeleton* Skeleton = Context.AnimInstanceProxy->GetSkeleton();
@@ -22,31 +34,24 @@ void FAnimNode_PoseByName::Initialize(const FAnimationInitializeContext& Context
 	}
 }
 
-void FAnimNode_PoseByName::CacheBones(const FAnimationCacheBonesContext& Context) 
-{
-}
-
-void FAnimNode_PoseByName::UpdateAssetPlayer(const FAnimationUpdateContext& Context)
-{
-	EvaluateGraphExposedInputs.Execute(Context);
-}
-
 void FAnimNode_PoseByName::Evaluate(FPoseContext& Output)
 {
-	if ((PoseAsset!= NULL) && (Output.AnimInstanceProxy->IsSkeletonCompatible(PoseAsset->GetSkeleton())))
+	if ((CurrentPoseAsset.IsValid()) && (Output.AnimInstanceProxy->IsSkeletonCompatible(CurrentPoseAsset->GetSkeleton())))
 	{
-		FBlendedCurve PoseCurve;
-		PoseCurve.InitFrom(Output.Curve);
+		// Zero out all curve value
+		FMemory::Memzero(PoseExtractContext.PoseCurves.GetData(), PoseExtractContext.PoseCurves.GetAllocatedSize());
 
 		if (PoseUID != FSmartNameMapping::MaxUID)
 		{
-			PoseCurve.Set(PoseUID, 1.f, ACF_DrivesPose);
+			int32 Index = PoseUIDList.Find(PoseUID);
+			if (Index!=INDEX_NONE)
+			{
+				PoseExtractContext.PoseCurves[Index] = PoseWeight;
+			}
 		}
 
 		// only give pose curve, we don't set any more curve here
-		PoseAsset->GetAnimationPose(Output.Pose, PoseCurve, FAnimExtractContext());
-
-		// this curve data doesn't go out, but it is intentional
+		CurrentPoseAsset->GetAnimationPose(Output.Pose, Output.Curve, PoseExtractContext);
 	}
 	else
 	{
@@ -54,18 +59,10 @@ void FAnimNode_PoseByName::Evaluate(FPoseContext& Output)
 	}
 }
 
-void FAnimNode_PoseByName::OverrideAsset(UAnimationAsset* NewAsset)
-{
-	if(UPoseAsset* NewPoseAsset = Cast<UPoseAsset>(NewAsset))
-	{
-		PoseAsset = NewPoseAsset;
-	}
-}
-
 void FAnimNode_PoseByName::GatherDebugData(FNodeDebugData& DebugData)
 {
 	FString DebugLine = DebugData.GetNodeName(this);
 	
-	DebugLine += FString::Printf(TEXT("('%s' Pose: %s)"), *GetNameSafe(PoseAsset), *PoseName.ToString());
+	DebugLine += FString::Printf(TEXT("('%s' Pose: %s)"), CurrentPoseAsset.IsValid()? *CurrentPoseAsset.Get()->GetName() : TEXT("None"), *PoseName.ToString());
 	DebugData.AddDebugItem(DebugLine, true);
 }
\ No newline at end of file
diff --git a/Engine/Source/Runtime/AnimGraphRuntime/Private/AnimNodes/AnimNode_PoseHandler.cpp b/Engine/Source/Runtime/AnimGraphRuntime/Private/AnimNodes/AnimNode_PoseHandler.cpp
new file mode 100644
--- /dev/null
+++ b/Engine/Source/Runtime/AnimGraphRuntime/Private/AnimNodes/AnimNode_PoseHandler.cpp
@@ -0,0 +1,99 @@
+// Copyright 1998-2016 Epic Games, Inc. All Rights Reserved.
+
+#include "AnimGraphRuntimePrivatePCH.h"
+#include "AnimNodes/AnimNode_PoseHandler.h"
+#include "Animation/AnimInstanceProxy.h"
+
+/////////////////////////////////////////////////////
+// FAnimPoseByNameNode
+
+void FAnimNode_PoseHandler::Initialize(const FAnimationInitializeContext& Context)
+{
+	FAnimNode_AssetPlayerBase::Initialize(Context);
+
+	UpdatePoseAssetProperty(Context.AnimInstanceProxy);
+}
+
+void FAnimNode_PoseHandler::CacheBones(const FAnimationCacheBonesContext& Context) 
+{
+	FAnimNode_AssetPlayerBase::CacheBones(Context);
+
+	BoneBlendWeights.Reset();
+	// this has to update bone blending weight
+	if (CurrentPoseAsset.IsValid())
+	{
+		const TArray<FName> TrackNames = CurrentPoseAsset.Get()->GetTrackNames();
+		const FBoneContainer& BoneContainer = Context.AnimInstanceProxy->GetRequiredBones();
+		const TArray<FBoneIndexType>& RequiredBoneIndices = BoneContainer.GetBoneIndicesArray();
+		BoneBlendWeights.AddZeroed(RequiredBoneIndices.Num());
+
+		for (const auto& TrackName : TrackNames)
+		{
+			int32 MeshBoneIndex = BoneContainer.GetPoseBoneIndexForBoneName(TrackName);
+			if (MeshBoneIndex != INDEX_NONE)
+			{
+				BoneBlendWeights[MeshBoneIndex] = 1.f;
+			}
+		}
+	}
+}
+
+void FAnimNode_PoseHandler::UpdateAssetPlayer(const FAnimationUpdateContext& Context)
+{
+	EvaluateGraphExposedInputs.Execute(Context);
+
+	// update pose asset if it's not valid
+	if (CurrentPoseAsset.IsValid() == false || CurrentPoseAsset.Get() != PoseAsset)
+	{
+		UpdatePoseAssetProperty(Context.AnimInstanceProxy);
+	}
+}
+
+void FAnimNode_PoseHandler::OverrideAsset(UAnimationAsset* NewAsset)
+{
+	if(UPoseAsset* NewPoseAsset = Cast<UPoseAsset>(NewAsset))
+	{
+		PoseAsset = NewPoseAsset;
+	}
+}
+
+void FAnimNode_PoseHandler::GatherDebugData(FNodeDebugData& DebugData)
+{
+	FString DebugLine = DebugData.GetNodeName(this);
+	
+	DebugLine += FString::Printf(TEXT("('%s')"), *GetNameSafe(PoseAsset));
+	DebugData.AddDebugItem(DebugLine, true);
+}
+
+void FAnimNode_PoseHandler::UpdatePoseAssetProperty(struct FAnimInstanceProxy* InstanceProxy)
+{
+	CurrentPoseAsset = PoseAsset;
+
+	if (CurrentPoseAsset.IsValid())
+	{
+		const TArray<FSmartName>& PoseNames = CurrentPoseAsset->GetPoseNames();
+		const int32 TotalPoseNum = PoseNames.Num();
+		if (TotalPoseNum > 0)
+		{
+			PoseExtractContext.PoseCurves.Reset(TotalPoseNum);
+			PoseUIDList.Reset(TotalPoseNum);
+
+			if (TotalPoseNum > 0)
+			{
+				PoseExtractContext.PoseCurves.AddZeroed(TotalPoseNum);
+				for (const auto& PoseName : PoseNames)
+				{
+					PoseUIDList.Add(PoseName.UID);
+				}
+			}
+
+			CacheBones(FAnimationCacheBonesContext(InstanceProxy));
+		}
+	}
+	else
+	{
+		PoseUIDList.Reset();
+		PoseExtractContext.PoseCurves.Reset();
+		BoneBlendWeights.Reset();
+	}
+}
\ No newline at end of file
diff --git a/Engine/Source/Runtime/AnimGraphRuntime/Private/AnimNodes/AnimNode_RefPose.cpp b/Engine/Source/Runtime/AnimGraphRuntime/Private/AnimNodes/AnimNode_RefPose.cpp
--- a/Engine/Source/Runtime/AnimGraphRuntime/Private/AnimNodes/AnimNode_RefPose.cpp
+++ b/Engine/Source/Runtime/AnimGraphRuntime/Private/AnimNodes/AnimNode_RefPose.cpp
@@ -19,7 +19,7 @@ void FAnimNode_RefPose::Evaluate(FPoseContext& Output)
 
 	case EIT_Additive:
 	default:
-		Output.ResetToIdentity();
+		Output.ResetToAdditiveIdentity();
 		break;
 	}
 }
diff --git a/Engine/Source/Runtime/AnimGraphRuntime/Private/BoneControllers/AnimNode_BoneDrivenController.cpp b/Engine/Source/Runtime/AnimGraphRuntime/Private/BoneControllers/AnimNode_BoneDrivenController.cpp
--- a/Engine/Source/Runtime/AnimGraphRuntime/Private/BoneControllers/AnimNode_BoneDrivenController.cpp
+++ b/Engine/Source/Runtime/AnimGraphRuntime/Private/BoneControllers/AnimNode_BoneDrivenController.cpp
@@ -188,7 +188,7 @@ void FAnimNode_BoneDrivenController::EvaluateComponentSpaceInternal(FComponentSp
 		FSmartNameMapping::UID NameUID = Skeleton->GetUIDByName(USkeleton::AnimCurveMappingName, ParameterName);
 		if (NameUID != FSmartNameMapping::MaxUID)
 		{
-			Context.Curve.Set(NameUID, FinalDriverValue, (DestinationMode == EDrivenDestinationMode::MorphTarget) ? EAnimCurveFlags::ACF_DrivesMorphTarget : EAnimCurveFlags::ACF_DrivesMaterial);
+			Context.Curve.Set(NameUID, FinalDriverValue, (DestinationMode == EDrivenDestinationMode::MorphTarget) ? EAnimCurveFlags::ACF_DriveMorphTarget : EAnimCurveFlags::ACF_DriveMaterial);
 		}
 	}
 }
diff --git a/Engine/Source/Runtime/AnimGraphRuntime/Private/BoneControllers/AnimNode_OrientationDriver.cpp b/Engine/Source/Runtime/AnimGraphRuntime/Private/BoneControllers/AnimNode_OrientationDriver.cpp
--- a/Engine/Source/Runtime/AnimGraphRuntime/Private/BoneControllers/AnimNode_OrientationDriver.cpp
+++ b/Engine/Source/Runtime/AnimGraphRuntime/Private/BoneControllers/AnimNode_OrientationDriver.cpp
@@ -2,6 +2,7 @@
 
 #include "AnimGraphRuntimePrivatePCH.h"
 #include "BoneControllers/AnimNode_OrientationDriver.h"
+#include "Animation/PoseAsset.h"
 
 #include "AnimInstanceProxy.h"
 
@@ -14,9 +15,9 @@ void FOrientationDriverParamSet::AddParam(const FOrientationDriverParam& InParam
 	// Look to see if this param already exists in this set
 	for (FOrientationDriverParam& Param : Params)
 	{
-		if (Param.ParameterName == InParam.ParameterName)
+		if (Param.ParamInfo.Name == InParam.ParamInfo.Name)
 		{
-			Param.ParameterValue += (InParam.ParameterValue * InScale);
+			Param.ParamValue += (InParam.ParamValue * InScale);
 			bFoundParam = true;
 		}
 	}
@@ -26,8 +27,8 @@ void FOrientationDriverParamSet::AddParam(const FOrientationDriverParam& InParam
 	{
 		int32 NewParamIndex = Params.AddZeroed();
 		FOrientationDriverParam& Param = Params[NewParamIndex];
-		Param.ParameterName = InParam.ParameterName;
-		Param.ParameterValue = (InParam.ParameterValue * InScale);
+		Param.ParamInfo = InParam.ParamInfo;
+		Param.ParamValue = (InParam.ParamValue * InScale);
 	}
 }
 
@@ -43,7 +44,7 @@ void FOrientationDriverParamSet::ScaleAllParams(float InScale)
 {
 	for (FOrientationDriverParam& Param : Params)
 	{
-		Param.ParameterValue *= InScale;
+		Param.ParamValue *= InScale;
 	}
 }
 
@@ -58,6 +59,7 @@ void FOrientationDriverParamSet::ClearParams()
 FAnimNode_OrientationDriver::FAnimNode_OrientationDriver()
 {
 	RadialScaling = 0.25f;
+	bIncludeRefPoseAsNeutralPose = true;
 }
 
 void FAnimNode_OrientationDriver::GatherDebugData(FNodeDebugData& DebugData)
@@ -74,43 +76,123 @@ void FAnimNode_OrientationDriver::EvaluateBoneTransforms(USkeletalMeshComponent*
 
 }
 
+FVector FAnimNode_OrientationDriver::GetTwistAxisVector()
+{
+	switch (TwistAxis)
+	{
+	case BA_X:
+	default:
+		return FVector(1.f, 0.f, 0.f);
+	case BA_Y:
+		return FVector(0.f, 1.f, 0.f);
+	case BA_Z:
+		return FVector(0.f, 0.f, 1.f);
+	}
+}
 
-
-void FAnimNode_OrientationDriver::UpdateMaxDistanceBetweenPoses()
+float FAnimNode_OrientationDriver::FindDistanceBetweenPoses(const FQuat& A, const FQuat& B)
 {
-	// Find largest distance between poses
-	MaxDistanceBetweenPoses = 0.f;
+	if (Type == EOrientationDriverType::SwingAndTwist)
+	{
+		return A.AngularDistance(B);
+	}
+	else
+	{
+		FVector TwistVector = GetTwistAxisVector();
+		FVector VecA = A.RotateVector(TwistVector);
+		FVector VecB = B.RotateVector(TwistVector);
+
+		const float Dot = FVector::DotProduct(VecA, VecB);
+		return FMath::Acos(Dot);
+	}
+}
 
-	for (int32 I = 1; I < Poses.Num(); I++)
+void FAnimNode_OrientationDriver::UpdateCachedPoseInfo(const FQuat& RefQuat)
+{
+	if (PoseSource != nullptr)
 	{
-		FQuat QuatI = Poses[I].PoseQuat;
-		for (int32 J = 0; J < I; J++)
+		int32 NumPoses = PoseSource->GetNumPoses();
+		int32 NumPoseInfos = NumPoses;
+		if (bIncludeRefPoseAsNeutralPose)
+		{
+			NumPoseInfos++;
+		}
+
+		PoseInfos.Empty();
+		PoseInfos.AddZeroed(NumPoseInfos);
+
+		// Get track index in PoseAsset for bone of interest
+		const int32 TrackIndex = PoseSource->GetTrackIndexByName(SourceBone.BoneName);
+		if (TrackIndex != INDEX_NONE)
 		{
-			FQuat QuatJ = Poses[J].PoseQuat;
+			TArray<FSmartName> PoseNames = PoseSource->GetPoseNames();
+
+			// Cache quat for source bone for each pose
+			for (int32 PoseIdx = 0; PoseIdx < NumPoses; PoseIdx++)
+			{
+				FOrientationDriverPoseInfo& PoseInfo = PoseInfos[PoseIdx];
+
+				FTransform BoneTransform;
+				bool bFound = PoseSource->GetLocalPoseForTrack(PoseIdx, TrackIndex, BoneTransform);
+				check(bFound);
+
+				PoseInfo.PoseQuat = BoneTransform.GetRotation();
+				PoseInfo.PoseName = PoseNames[PoseIdx].DisplayName;
+			}
 
-			float Dist = QuatI.AngularDistance(QuatJ);
+			// If we want to include ref pose, add that to end
+			if (bIncludeRefPoseAsNeutralPose)
+			{
+				FOrientationDriverPoseInfo& PoseInfo = PoseInfos.Last();
+
+				static FName RefPoseName = FName(TEXT("RefPose"));
+				PoseInfo.PoseName = RefPoseName;
+				PoseInfo.PoseQuat = RefQuat;
+			}
 
-			MaxDistanceBetweenPoses = FMath::Max(Dist, MaxDistanceBetweenPoses);
+			// Calculate largest distance between poses
+			if (NumPoseInfos == 1)
+			{
+				// If only one pose, set to 180 degrees
+				PoseInfos[0].NearestPoseDist = PI;
+			}
+			else
+			{
+				// Iterate over poses
+				for (int32 PoseIdx = 0; PoseIdx < NumPoseInfos; PoseIdx++)
+				{
+					FOrientationDriverPoseInfo& PoseInfo = PoseInfos[PoseIdx];
+					PoseInfo.NearestPoseDist = BIG_NUMBER; // init to large value
+
+					for (int32 OtherPoseIdx = 0; OtherPoseIdx < NumPoseInfos; OtherPoseIdx++)
+					{
+						if (OtherPoseIdx != PoseIdx) // If not ourself..
+						{
+							// Get distance between poses
+							FOrientationDriverPoseInfo& OtherPoseInfo = PoseInfos[OtherPoseIdx];
+							float Dist = FindDistanceBetweenPoses(PoseInfo.PoseQuat, OtherPoseInfo.PoseQuat);
+							PoseInfo.NearestPoseDist = FMath::Min(Dist, PoseInfo.NearestPoseDist);
+						}
+					}
+
+					// Avoid zero dist if poses are all on top of each other
+					PoseInfo.NearestPoseDist = FMath::Max(PoseInfo.NearestPoseDist, KINDA_SMALL_NUMBER);
+				}
+			}
 		}
 	}
 
-	bMaxDistanceUpToDate = true;
+	bCachedPoseInfoUpToDate = true;
 }
 
 void FAnimNode_OrientationDriver::EvaluateComponentSpaceInternal(FComponentSpacePoseContext& Context)
 {
-	// Do nothing if no poses
-	if (Poses.Num() == 0)
+	// Do nothing if no PoseAsset, or Info doesn't match
+	if (PoseSource == nullptr)
 	{
 		return;
 	}
 
-	// Ensure max distance up to date
-	if (!bMaxDistanceUpToDate)
-	{
-		UpdateMaxDistanceBetweenPoses();
-	}
-
 	const FBoneContainer& BoneContainer = Context.Pose.GetPose().GetBoneContainer();
 
 	// Final set of driven params, from all interpolators
@@ -120,25 +202,62 @@ void FAnimNode_OrientationDriver::EvaluateComponentSpaceInternal(FComponentSpace
 	const FTransform SourceCurrentBoneTransform = Context.Pose.GetLocalSpaceTransform(SourceBone.GetCompactPoseIndex(BoneContainer));
 	const FQuat BoneQuat = SourceCurrentBoneTransform.GetRotation();
 
+	const TArray<FAnimCurveBase> PoseAssetCurveData = PoseSource->GetCurveData();
+
 	// Array of 'weights' of current orientation to each pose
 	float TotalWeight = 0.f;
 
-	float GaussVarianceSqr = FMath::Square(RadialScaling * MaxDistanceBetweenPoses);
+	// Ensure cached info is up to date
+	if (!bCachedPoseInfoUpToDate)
+	{
+		const FTransform& SourceRefPoseBoneTransform = BoneContainer.GetRefPoseArray()[SourceBone.BoneIndex];
+
+		UpdateCachedPoseInfo(SourceRefPoseBoneTransform.GetRotation());
+	}
 
 	// Iterate over each pose, adding its contribution
-	for (FOrientationDriverPose& Pose : Poses)
+	for (int32 PoseInfoIdx = 0; PoseInfoIdx < PoseInfos.Num(); PoseInfoIdx++)
 	{
+		FOrientationDriverPoseInfo& PoseInfo = PoseInfos[PoseInfoIdx];
+
 		// Find distance
-		Pose.PoseDistance = BoneQuat.AngularDistance(Pose.PoseQuat);
+		PoseInfo.PoseDistance = FindDistanceBetweenPoses(BoneQuat, PoseInfo.PoseQuat);
 
 		// Evaluate radial basis function to find weight
-		Pose.PoseWeight = FMath::Exp(-(Pose.PoseDistance * Pose.PoseDistance) / GaussVarianceSqr);
+		float GaussVarianceSqr = FMath::Square(RadialScaling * PoseInfo.NearestPoseDist);
+		PoseInfo.PoseWeight = FMath::Exp(-(PoseInfo.PoseDistance * PoseInfo.PoseDistance) / GaussVarianceSqr);
 
 		// Add weight to total
-		TotalWeight += Pose.PoseWeight;
+		TotalWeight += PoseInfo.PoseWeight;
+
+		// Final pose info may be ref pose, will not exist in PoseSource
+		if (PoseInfoIdx < PoseSource->GetNumPoses())
+		{
+			// Get param values for this pose
+			TArray<float> PoseAssetCurveValues = PoseSource->GetCurveValues(PoseInfoIdx);
+			check(PoseAssetCurveValues.Num() == PoseAssetCurveData.Num()); // Should match name array
 
-		// Add params for this pose, weighted by pose weight
-		ResultParamSet.AddParams(Pose.DrivenParams, Pose.PoseWeight);
+			// Add param values to driver info for this pose
+			for (int32 CurveIdx = 0; CurveIdx < PoseAssetCurveValues.Num(); CurveIdx++)
+			{
+				// Only both copying params with non-zero weights
+				float ParamValue = PoseAssetCurveValues[CurveIdx];
+				if (FMath::Abs(ParamValue) > KINDA_SMALL_NUMBER)
+				{
+					FOrientationDriverParam NewParam;
+
+					NewParam.ParamInfo.Name = PoseAssetCurveData[CurveIdx].Name.DisplayName;
+					NewParam.ParamInfo.UID = PoseAssetCurveData[CurveIdx].Name.UID;
+					NewParam.ParamInfo.Type.bMaterial = PoseAssetCurveData[CurveIdx].GetCurveTypeFlag(ACF_DriveMaterial);
+					NewParam.ParamInfo.Type.bMorphtarget = PoseAssetCurveData[CurveIdx].GetCurveTypeFlag(ACF_DriveMorphTarget);
+					NewParam.ParamInfo.Type.bEvent = false;
+
+					NewParam.ParamValue = PoseAssetCurveValues[CurveIdx];
+
+					ResultParamSet.AddParam(NewParam, PoseInfo.PoseWeight);
+				}
+			}
+		}
 	}
 
 	// Only normalize and apply if we got some kind of weight
@@ -149,9 +268,9 @@ void FAnimNode_OrientationDriver::EvaluateComponentSpaceInternal(FComponentSpace
 		ResultParamSet.ScaleAllParams(WeightScale);
 
 		// Also normalize each pose weight
-		for (FOrientationDriverPose& Pose : Poses)
+		for (FOrientationDriverPoseInfo& PoseInfo : PoseInfos)
 		{
-			Pose.PoseWeight *= WeightScale;
+			PoseInfo.PoseWeight *= WeightScale;
 		}
 
 		//	Morph target and Material parameter curves
@@ -159,10 +278,9 @@ void FAnimNode_OrientationDriver::EvaluateComponentSpaceInternal(FComponentSpace
 
 		for (FOrientationDriverParam& Param : ResultParamSet.Params)
 		{
-			FSmartNameMapping::UID NameUID = Skeleton->GetUIDByName(USkeleton::AnimCurveMappingName, Param.ParameterName);
-			if (NameUID != FSmartNameMapping::MaxUID)
+			if (Param.ParamInfo.UID != FSmartNameMapping::MaxUID)
 			{
-				Context.Curve.Set(NameUID, Param.ParameterValue, EAnimCurveFlags::ACF_DrivesMorphTarget | EAnimCurveFlags::ACF_DrivesMaterial | EAnimCurveFlags::ACF_DrivesPose);
+				Context.Curve.Set(Param.ParamInfo.UID, Param.ParamValue, Param.ParamInfo.Type.GetAnimCurveFlags() | ACF_DriveAttribute);
 			}
 		}
 	}
@@ -178,4 +296,4 @@ void FAnimNode_OrientationDriver::InitializeBoneReferences(const FBoneContainer&
 {
 	// Init bone ref
 	SourceBone.Initialize(RequiredBones);
-}
\ No newline at end of file
+}
diff --git a/Engine/Source/Runtime/AnimGraphRuntime/Public/AnimNodes/AnimNode_PoseBlendNode.h b/Engine/Source/Runtime/AnimGraphRuntime/Public/AnimNodes/AnimNode_PoseBlendNode.h
--- a/Engine/Source/Runtime/AnimGraphRuntime/Public/AnimNodes/AnimNode_PoseBlendNode.h
+++ b/Engine/Source/Runtime/AnimGraphRuntime/Public/AnimNodes/AnimNode_PoseBlendNode.h
@@ -3,47 +3,33 @@
 #pragma once
 #include "Animation/AnimNodeBase.h"
 #include "Animation/PoseAsset.h"
-#include "Animation/AnimNode_AssetPlayerBase.h"
+#include "AnimNode_PoseHandler.h"
 #include "AnimNode_PoseBlendNode.generated.h"
 
 // Evaluates a point in an anim sequence, using a specific time input rather than advancing time internally.
 // Typically the playback position of the animation for this node will represent something other than time, like jump height.
 // This node will not trigger any notifies present in the associated sequence.
 USTRUCT()
-struct ANIMGRAPHRUNTIME_API FAnimNode_PoseBlendNode : public FAnimNode_AssetPlayerBase
+struct ANIMGRAPHRUNTIME_API FAnimNode_PoseBlendNode : public FAnimNode_PoseHandler
 {
 	GENERATED_USTRUCT_BODY()
 public:
 	UPROPERTY(EditAnywhere, EditFixedSize, BlueprintReadWrite, Category = Links)
 	FPoseLink SourcePose;
 
-	// The animation sequence asset to evaluate
-	UPROPERTY(EditAnywhere, BlueprintReadWrite, Category = Settings, meta = (PinHiddenByDefault))
-	UPoseAsset* PoseAsset;
-
 public:	
 	FAnimNode_PoseBlendNode()
 	{
 	}
 
-	// FAnimNode_AssetPlayerBase interface
-	virtual float GetCurrentAssetTime() { return 0.f; }
-	virtual float GetCurrentAssetLength() { return 0.f; }
-	// End of FAnimNode_AssetPlayerBase interface
 
 	// FAnimNode_Base interface
 	virtual void Initialize(const FAnimationInitializeContext& Context) override;
 	virtual void CacheBones(const FAnimationCacheBonesContext& Context) override;
 	virtual void UpdateAssetPlayer(const FAnimationUpdateContext& Context) override;
 	virtual void Evaluate(FPoseContext& Output) override;
-	virtual void OverrideAsset(UAnimationAsset* NewAsset) override;
 	virtual void GatherDebugData(FNodeDebugData& DebugData) override;
 	// End of FAnimNode_Base interface
 
-	// FAnimNode_AssetPlayerBase Interface
-	virtual float GetAccumulatedTime() {return 0.f;}
-	virtual void SetAccumulatedTime(const float& NewTime) {}
-	virtual UAnimationAsset* GetAnimAsset() {return PoseAsset;}
-	// End of FAnimNode_AssetPlayerBase Interface
 };
 
diff --git a/Engine/Source/Runtime/AnimGraphRuntime/Public/AnimNodes/AnimNode_PoseByName.h b/Engine/Source/Runtime/AnimGraphRuntime/Public/AnimNodes/AnimNode_PoseByName.h
--- a/Engine/Source/Runtime/AnimGraphRuntime/Public/AnimNodes/AnimNode_PoseByName.h
+++ b/Engine/Source/Runtime/AnimGraphRuntime/Public/AnimNodes/AnimNode_PoseByName.h
@@ -3,50 +3,37 @@
 #pragma once
 #include "Animation/AnimNodeBase.h"
 #include "Animation/PoseAsset.h"
-#include "Animation/AnimNode_AssetPlayerBase.h"
+#include "AnimNode_PoseHandler.h"
 #include "AnimNode_PoseByName.generated.h"
 
 // Evaluates a point in an anim sequence, using a specific time input rather than advancing time internally.
 // Typically the playback position of the animation for this node will represent something other than time, like jump height.
 // This node will not trigger any notifies present in the associated sequence.
 USTRUCT()
-struct ANIMGRAPHRUNTIME_API FAnimNode_PoseByName : public FAnimNode_AssetPlayerBase
+struct ANIMGRAPHRUNTIME_API FAnimNode_PoseByName : public FAnimNode_PoseHandler
 {
 	GENERATED_USTRUCT_BODY()
 public:
-	// The animation sequence asset to evaluate
 	UPROPERTY(EditAnywhere, BlueprintReadWrite, Category = Settings, meta = (PinHiddenByDefault))
-	UPoseAsset* PoseAsset;
+	FName PoseName;
 
 	UPROPERTY(EditAnywhere, BlueprintReadWrite, Category = Settings, meta = (PinHiddenByDefault))
-	FName PoseName;
+	float PoseWeight;
 
 public:	
 	FAnimNode_PoseByName()
-		: PoseUID(FSmartNameMapping::MaxUID)
+		: PoseWeight(1.f)
+		, PoseUID(FSmartNameMapping::MaxUID)
 	{
 	}
 
-	// FAnimNode_AssetPlayerBase interface
-	virtual float GetCurrentAssetTime() { return 0.f; }
-	virtual float GetCurrentAssetLength() { return 0.f; }
-	// End of FAnimNode_AssetPlayerBase interface
-
 	// FAnimNode_Base interface
 	virtual void Initialize(const FAnimationInitializeContext& Context) override;
 	virtual void CacheBones(const FAnimationCacheBonesContext& Context) override;
 	virtual void UpdateAssetPlayer(const FAnimationUpdateContext& Context) override;
 	virtual void Evaluate(FPoseContext& Output) override;
-	virtual void OverrideAsset(UAnimationAsset* NewAsset) override;
 	virtual void GatherDebugData(FNodeDebugData& DebugData) override;
 	// End of FAnimNode_Base interface
-
-	// FAnimNode_AssetPlayerBase Interface
-	virtual float GetAccumulatedTime() {return 0.f;}
-	virtual void SetAccumulatedTime(const float& NewTime) {}
-	virtual UAnimationAsset* GetAnimAsset() {return PoseAsset;}
-	// End of FAnimNode_AssetPlayerBase Interface
-
 private:
 	// PoseUID
 	FSmartNameMapping::UID PoseUID;
diff --git a/Engine/Source/Runtime/AnimGraphRuntime/Public/AnimNodes/AnimNode_PoseHandler.h b/Engine/Source/Runtime/AnimGraphRuntime/Public/AnimNodes/AnimNode_PoseHandler.h
new file mode 100644
--- /dev/null
+++ b/Engine/Source/Runtime/AnimGraphRuntime/Public/AnimNodes/AnimNode_PoseHandler.h
@@ -0,0 +1,56 @@
+// Copyright 1998-2016 Epic Games, Inc. All Rights Reserved.
+
+#pragma once
+#include "Animation/AnimNodeBase.h"
+#include "Animation/PoseAsset.h"
+#include "Animation/AnimNode_AssetPlayerBase.h"
+#include "AnimNode_PoseHandler.generated.h"
+
+// Evaluates a point in an anim sequence, using a specific time input rather than advancing time internally.
+// Typically the playback position of the animation for this node will represent something other than time, like jump height.
+// This node will not trigger any notifies present in the associated sequence.
+USTRUCT()
+struct ANIMGRAPHRUNTIME_API FAnimNode_PoseHandler : public FAnimNode_AssetPlayerBase
+{
+	GENERATED_USTRUCT_BODY()
+public:
+	// The animation sequence asset to evaluate
+	UPROPERTY(EditAnywhere, BlueprintReadWrite, Category = Settings, meta = (PinHiddenByDefault))
+	UPoseAsset* PoseAsset;
+
+public:	
+	FAnimNode_PoseHandler()
+	{
+	}
+
+	// FAnimNode_AssetPlayerBase interface
+	virtual float GetCurrentAssetTime() { return 0.f; }
+	virtual float GetCurrentAssetLength() { return 0.f; }
+	// End of FAnimNode_AssetPlayerBase interface
+
+	// FAnimNode_Base interface
+	virtual void Initialize(const FAnimationInitializeContext& Context) override;
+	virtual void CacheBones(const FAnimationCacheBonesContext& Context) override;
+	virtual void UpdateAssetPlayer(const FAnimationUpdateContext& Context) override;
+	virtual void OverrideAsset(UAnimationAsset* NewAsset) override;
+	virtual void GatherDebugData(FNodeDebugData& DebugData) override;
+	// End of FAnimNode_Base interface
+
+	// FAnimNode_AssetPlayerBase Interface
+	virtual float GetAccumulatedTime() {return 0.f;}
+	virtual void SetAccumulatedTime(const float& NewTime) {}
+	virtual UAnimationAsset* GetAnimAsset() {return PoseAsset;}
+	// End of FAnimNode_AssetPlayerBase Interface
+
+protected:
+	TWeakObjectPtr<UPoseAsset> CurrentPoseAsset;
+	FAnimExtractContext PoseExtractContext;
+	// weight to blend pose per joint - has to be cached whenever cache bones for LOD
+	// note this is for mesh bone
+	TArray<float> BoneBlendWeights;
+	TArray<FSmartNameMapping::UID> PoseUIDList;
+
+private:
+	void UpdatePoseAssetProperty(struct FAnimInstanceProxy* InstanceProxy);
+};
+
diff --git a/Engine/Source/Runtime/AnimGraphRuntime/Public/BoneControllers/AnimNode_OrientationDriver.h b/Engine/Source/Runtime/AnimGraphRuntime/Public/BoneControllers/AnimNode_OrientationDriver.h
--- a/Engine/Source/Runtime/AnimGraphRuntime/Public/BoneControllers/AnimNode_OrientationDriver.h
+++ b/Engine/Source/Runtime/AnimGraphRuntime/Public/BoneControllers/AnimNode_OrientationDriver.h
@@ -6,29 +6,26 @@
 #include "AnimNode_OrientationDriver.generated.h"
 
 /** One named parameter being driven by the orientation of a bone. */
-USTRUCT()
 struct FOrientationDriverParam
 {
-	GENERATED_BODY()
-
-	UPROPERTY(EditAnywhere, Category = Param)
-	FName ParameterName;
+	/** Information about this curve (name etc) */
+	FAnimCurveParam ParamInfo;
 
-	UPROPERTY(EditAnywhere, Category = Param)
-	float ParameterValue;
+	/** Value of this parameter */
+	float ParamValue;
 };
 
 /** One target pose for the bone, with parameters to drive bone approaches this orientation. */
-USTRUCT()
-struct FOrientationDriverPose
+struct FOrientationDriverPoseInfo
 {
-	GENERATED_BODY()
+	/** Name of pose, from PoseAsset */
+	FName PoseName;
 
-	UPROPERTY(EditAnywhere, Category = Pose)
+	/** Rotation of source bone in this pose, from PoseAsset */
 	FQuat PoseQuat;
 
-	UPROPERTY(EditAnywhere, Category = Pose)
-	TArray<FOrientationDriverParam> DrivenParams;
+	/** Distance to nearest pose */
+	float NearestPoseDist;
 
 	/** Last weight calculated for this pose. */
 	float PoseWeight;
@@ -36,7 +33,7 @@ struct FOrientationDriverPose
 	/** Last distance calculated from this pose */
 	float PoseDistance;
 
-	FOrientationDriverPose()
+	FOrientationDriverPoseInfo()
 	: PoseQuat(FQuat::Identity)
 	{}
 };
@@ -60,6 +57,17 @@ struct FOrientationDriverParamSet
 	void ClearParams();
 };
 
+/** Orientation aspect used to drive interpolation */
+UENUM()
+enum class EOrientationDriverType : uint8
+{
+	/** Consider full rotation for interpolation */
+	SwingAndTwist,
+
+	/** Consider only swing for interpolation */
+	SwingOnly,
+};
+
 /** RBF based orientation driver */
 USTRUCT()
 struct ANIMGRAPHRUNTIME_API FAnimNode_OrientationDriver : public FAnimNode_SkeletalControlBase
@@ -70,17 +78,32 @@ struct ANIMGRAPHRUNTIME_API FAnimNode_OrientationDriver : public FAnimNode_Skele
 	UPROPERTY(EditAnywhere, Category = OrientationDriver)
 	FBoneReference SourceBone;
 
-	/** Set of target poses, which include parameters to drive. */
-	UPROPERTY(EditAnywhere, Category = OrientationDriver)
-	TArray<FOrientationDriverPose> Poses;
-
 	/** Scaling of radial basis, applied to max distance between poses */
 	UPROPERTY(EditAnywhere, Category = OrientationDriver, meta = (ClampMin = "0.01"))
 	float RadialScaling;
 
+	/** PoseAsset used as source of  */
+	UPROPERTY(EditAnywhere, Category = OrientationDriver)
+	UPoseAsset* PoseSource;
+
+	/** Should we consider the mesh ref pose of SourceBone as a 'neutral' pose (zero curves)   */
+	UPROPERTY(EditAnywhere, Category = OrientationDriver)
+	bool bIncludeRefPoseAsNeutralPose;
+
+	/** Type of orientation for driving parameter  */
+	UPROPERTY(EditAnywhere, Category = OrientationDriver)
+	EOrientationDriverType Type;
+
+	/** Axis to use when Type is SwingOnly */
+	UPROPERTY(EditAnywhere, Category = OrientationDriver)
+	TEnumAsByte<EBoneAxis> TwistAxis;
+
 	/** Previous set of interpolated params, kept for debugging */
 	FOrientationDriverParamSet ResultParamSet;
 
+	/** Cached set of info for each pose */
+	TArray<FOrientationDriverPoseInfo> PoseInfos;
+
 	// FAnimNode_Base interface
 	virtual void GatherDebugData(FNodeDebugData& DebugData) override;
 	// End of FAnimNode_Base interface
@@ -94,15 +117,18 @@ struct ANIMGRAPHRUNTIME_API FAnimNode_OrientationDriver : public FAnimNode_Skele
 	// End of FAnimNode_SkeletalControlBase interface
 
 	/** Find max distance between any of the poses */
-	void UpdateMaxDistanceBetweenPoses();
+	void UpdateCachedPoseInfo(const FQuat& RefQuat);
 
-protected:
+	/** Util to return unit vector for curent twist axis */
+	FVector GetTwistAxisVector();
+	/** Util to find distance between 2 quats, using Type and TwistAxis settings */
+	float FindDistanceBetweenPoses(const FQuat& A, const FQuat& B);
 
-	/** Is MaxDistanceBetweenPoses up to date */
-	bool bMaxDistanceUpToDate;
 
-	/** Maximum angular distance between poses in the Poses array */
-	float MaxDistanceBetweenPoses;
+protected:
+
+	/** Is cached info up to date */
+	bool bCachedPoseInfoUpToDate;
 
 	// FAnimNode_SkeletalControlBase protected interface
 	virtual void InitializeBoneReferences(const FBoneContainer& RequiredBones) override;
diff --git a/Engine/Source/Runtime/AnimGraphRuntime/Public/BoneControllers/AnimNode_RotationMultiplier.h b/Engine/Source/Runtime/AnimGraphRuntime/Public/BoneControllers/AnimNode_RotationMultiplier.h
--- a/Engine/Source/Runtime/AnimGraphRuntime/Public/BoneControllers/AnimNode_RotationMultiplier.h
+++ b/Engine/Source/Runtime/AnimGraphRuntime/Public/BoneControllers/AnimNode_RotationMultiplier.h
@@ -5,13 +5,6 @@
 #include "AnimNode_SkeletalControlBase.h"
 #include "AnimNode_RotationMultiplier.generated.h"
 
-UENUM()
-enum EBoneAxis
-{
-	BA_X,
-	BA_Y,
-	BA_Z
-};
 
 /**
  *	Simple controller that multiplies scalar value to the translation/rotation/scale of a single bone.
diff --git a/Engine/Source/Runtime/Core/Private/HAL/ThreadingBase.cpp b/Engine/Source/Runtime/Core/Private/HAL/ThreadingBase.cpp
--- a/Engine/Source/Runtime/Core/Private/HAL/ThreadingBase.cpp
+++ b/Engine/Source/Runtime/Core/Private/HAL/ThreadingBase.cpp
@@ -28,6 +28,14 @@ CORE_API bool IsInSlateThread()
 	return GSlateLoadingThreadId != 0 && FPlatformTLS::GetCurrentThreadId() == GSlateLoadingThreadId;
 }
 
+CORE_API FRunnableThread* GAudioThread = nullptr;
+
+CORE_API bool IsInAudioThread()
+{
+	// True if this is the audio thread or if there is no audio thread, then if it is the game thread
+	return (GAudioThreadId != 0 && FPlatformTLS::GetCurrentThreadId() == GAudioThreadId) || (GAudioThreadId == 0 && FPlatformTLS::GetCurrentThreadId() == GGameThreadId);
+}
+
 CORE_API int32 GIsRenderingThreadSuspended = 0;
 
 CORE_API FRunnableThread* GRenderingThread = nullptr;
diff --git a/Engine/Source/Runtime/Core/Private/IOS/IOSPlatformProcess.cpp b/Engine/Source/Runtime/Core/Private/IOS/IOSPlatformProcess.cpp
--- a/Engine/Source/Runtime/Core/Private/IOS/IOSPlatformProcess.cpp
+++ b/Engine/Source/Runtime/Core/Private/IOS/IOSPlatformProcess.cpp
@@ -78,9 +78,9 @@ void FIOSPlatformProcess::SetRealTimeMode()
 	}
 }
 
-void FIOSPlatformProcess::SetupGameOrRenderThread(bool bIsRenderThread)
+// Set the game thread priority to very high, slightly above the render thread
+void FIOSPlatformProcess::SetupThread(const int Priority)
 {
-	// Set the gamethread priority to very high, slightly above the renderthread
 	struct sched_param Sched;
 	FMemory::Memzero(&Sched, sizeof(struct sched_param));
 
@@ -90,10 +90,19 @@ void FIOSPlatformProcess::SetupGameOrRenderThread(bool bIsRenderThread)
 
 	// Set the new priority and policy (apple recommended FIFO for the two main non-working threads)
 	int32 Policy = SCHED_FIFO;
-	Sched.sched_priority = bIsRenderThread ? RENDER_THREAD_PRIORITY : GAME_THREAD_PRIORITY;
+	Sched.sched_priority = Priority;
 	pthread_setschedparam(pthread_self(), Policy, &Sched);
 }
 
+void FIOSPlatformProcess::SetupGameThread()
+{
+	SetupThread(GAME_THREAD_PRIORITY);
+}
+
+void FIOSPlatformProcess::SetupRenderThread()
+{
+	SetupThread(RENDER_THREAD_PRIORITY);
+}
 
 void FIOSPlatformProcess::SetThreadAffinityMask(uint64 AffinityMask)
 {
diff --git a/Engine/Source/Runtime/Core/Private/Misc/Core.cpp b/Engine/Source/Runtime/Core/Private/Misc/Core.cpp
--- a/Engine/Source/Runtime/Core/Private/Misc/Core.cpp
+++ b/Engine/Source/Runtime/Core/Private/Misc/Core.cpp
@@ -236,6 +236,7 @@ bool					GUseSeekFreeLoading				= false;
 uint32					GGameThreadId					= 0;
 uint32					GRenderThreadId					= 0;
 uint32					GSlateLoadingThreadId			= 0;
+uint32					GAudioThreadId					= 0;
 /** Has GGameThreadId been set yet?																			*/
 bool					GIsGameThreadIdInitialized		= false;
 
diff --git a/Engine/Source/Runtime/Core/Public/CoreGlobals.h b/Engine/Source/Runtime/Core/Public/CoreGlobals.h
--- a/Engine/Source/Runtime/Core/Public/CoreGlobals.h
+++ b/Engine/Source/Runtime/Core/Public/CoreGlobals.h
@@ -279,6 +279,9 @@ extern CORE_API uint32 GRenderThreadId;
 /** Thread ID of the slate thread, if any */
 extern CORE_API uint32 GSlateLoadingThreadId;
 
+/** Thread ID of the audio thread, if any */
+extern CORE_API uint32 GAudioThreadId;
+
 /** Has GGameThreadId been set yet? */
 extern CORE_API bool GIsGameThreadIdInitialized;
 
diff --git a/Engine/Source/Runtime/Core/Public/GenericPlatform/GenericPlatformMisc.h b/Engine/Source/Runtime/Core/Public/GenericPlatform/GenericPlatformMisc.h
--- a/Engine/Source/Runtime/Core/Public/GenericPlatform/GenericPlatformMisc.h
+++ b/Engine/Source/Runtime/Core/Public/GenericPlatform/GenericPlatformMisc.h
@@ -609,6 +609,17 @@ public:
 		return true;
 	}
 
+	/**
+	* Checks if platform wants to allow an audio thread on current device (note: does not imply it will, only if okay given other criteria met)
+	*
+	* @return true if allowed, false if shouldn't use a separate audio thread
+	*/
+	static bool AllowAudioThread()
+	{
+		// allow if not overridden
+		return true;
+	}
+
 	/**
 	 * Checks if platform wants to allow the thread heartbeat hang detection
 	 *
diff --git a/Engine/Source/Runtime/Core/Public/GenericPlatform/GenericPlatformProcess.h b/Engine/Source/Runtime/Core/Public/GenericPlatform/GenericPlatformProcess.h
--- a/Engine/Source/Runtime/Core/Public/GenericPlatform/GenericPlatformProcess.h
+++ b/Engine/Source/Runtime/Core/Public/GenericPlatform/GenericPlatformProcess.h
@@ -200,15 +200,15 @@ struct CORE_API FGenericPlatformProcess
 	 */
 	static void SetThreadAffinityMask( uint64 AffinityMask );
 
-	/**
-	 * Allow the platform to do anything it needs for game or render thread (this would use the NamedThreads enum if we could forward declare the enum).
-	 *
-	 * @param bIsRenderThread true if setting up the render thread, false for the game thread.
-	 */
-	static void SetupGameOrRenderThread(bool bIsRenderThread)
-	{
-	}
-	
+	/** Allow the platform to do anything it needs for game thread */
+	static void SetupGameThread() { }
+
+	/** Allow the platform to do anything it needs for render thread */
+	static void SetupRenderThread() { }
+
+	/** Allow the platform to do anything it needs for audio thread */
+	static void SetupAudioThread() { }
+
 	/** Get startup directory.  NOTE: Only one return value is valid at a time! **/
 	static const TCHAR* BaseDir();
 
diff --git a/Engine/Source/Runtime/Core/Public/HAL/ThreadingBase.h b/Engine/Source/Runtime/Core/Public/HAL/ThreadingBase.h
--- a/Engine/Source/Runtime/Core/Public/HAL/ThreadingBase.h
+++ b/Engine/Source/Runtime/Core/Public/HAL/ThreadingBase.h
@@ -1223,6 +1223,12 @@ FORCEINLINE bool IsInGameThread()
 	return true;
 }
 
+/** @return True if called from the audio thread, and not merely a thread calling audio functions. */
+extern CORE_API bool IsInAudioThread();
+
+/** Thread used for audio */
+extern CORE_API FRunnableThread* GAudioThread;
+
 /** @return True if called from the slate thread, and not merely a thread calling slate functions. */
 extern CORE_API bool IsInSlateThread();
 
diff --git a/Engine/Source/Runtime/Core/Public/IOS/IOSPlatformProcess.h b/Engine/Source/Runtime/Core/Public/IOS/IOSPlatformProcess.h
--- a/Engine/Source/Runtime/Core/Public/IOS/IOSPlatformProcess.h
+++ b/Engine/Source/Runtime/Core/Public/IOS/IOSPlatformProcess.h
@@ -33,9 +33,13 @@ struct CORE_API FIOSPlatformProcess : public FGenericPlatformProcess
 	static void LaunchURL( const TCHAR* URL, const TCHAR* Parms, FString* Error );
 	static bool CanLaunchURL(const TCHAR* URL);
 	static void SetRealTimeMode();
-	static void SetupGameOrRenderThread(bool bIsRenderThread);
+	static void SetupGameThread();
+	static void SetupRenderThread();
 	static void SetThreadAffinityMask(uint64 AffinityMask);
 	static const TCHAR* ExecutableName(bool bRemoveExtension = true);
+
+private:
+	static void SetupThread(int Priority);
 };
 
 typedef FIOSPlatformProcess FPlatformProcess;
